Interface level changes in preparation for Oculus 0.4 integration.

diff --git a/Engine/Source/Runtime/Core/Public/GenericPlatform/GenericWindow.h b/Engine/Source/Runtime/Core/Public/GenericPlatform/GenericWindow.h
--- a/Engine/Source/Runtime/Core/Public/GenericPlatform/GenericWindow.h
+++ b/Engine/Source/Runtime/Core/Public/GenericPlatform/GenericWindow.h
@@ -16,7 +16,12 @@ namespace EWindowMode
 		/** The window has no border and takes up the entire area of the screen */
 		WindowedFullscreen,
 		/** The window has a border and may not take up the entire screen area */
-		Windowed
+		Windowed,
+		/** Pseudo-fullscreen mode for devices like HMDs */
+		WindowedMirror,
+
+		/** The total number of supported window modes */
+		NumWindowModes
 	};
 
 	static inline Type ConvertIntToWindowMode(int32 InWindowMode)
@@ -30,6 +35,9 @@ namespace EWindowMode
 			case 1:
 				WindowMode = WindowedFullscreen;
 				break;
+			case 3:
+				WindowMode = WindowedMirror;
+				break;
 			case 2:
 			default:
 				WindowMode = Windowed;
diff --git a/Engine/Source/Runtime/Engine/Classes/GameFramework/GameUserSettings.h b/Engine/Source/Runtime/Engine/Classes/GameFramework/GameUserSettings.h
--- a/Engine/Source/Runtime/Engine/Classes/GameFramework/GameUserSettings.h
+++ b/Engine/Source/Runtime/Engine/Classes/GameFramework/GameUserSettings.h
@@ -143,6 +143,7 @@ protected:
 	 *	0 = Fullscreen
 	 *	1 = Windowed fullscreen
 	 *	2 = Windowed
+	 *	3 = WindowedMirror
 	 */
 	UPROPERTY(config)
 	int32 FullscreenMode;
diff --git a/Engine/Source/Runtime/Engine/Public/UnrealClient.h b/Engine/Source/Runtime/Engine/Public/UnrealClient.h
--- a/Engine/Source/Runtime/Engine/Public/UnrealClient.h
+++ b/Engine/Source/Runtime/Engine/Public/UnrealClient.h
@@ -281,7 +281,7 @@ public:
 	virtual float GetTabletPressure() { return 0.f; }
 	virtual bool IsPenActive() { return false; }
 	virtual void SetMouse(int32 x, int32 y) = 0;
-	virtual bool IsFullscreen()	const { return WindowMode == EWindowMode::Fullscreen || WindowMode == EWindowMode::WindowedFullscreen; }
+	virtual bool IsFullscreen()	const { return WindowMode == EWindowMode::Fullscreen || WindowMode == EWindowMode::WindowedFullscreen || WindowMode == EWindowMode::WindowedMirror; }
 	virtual void ProcessInput( float DeltaTime ) = 0;
 
 	/**
diff --git a/Engine/Source/Runtime/HeadMountedDisplay/Public/IHeadMountedDisplay.h b/Engine/Source/Runtime/HeadMountedDisplay/Public/IHeadMountedDisplay.h
--- a/Engine/Source/Runtime/HeadMountedDisplay/Public/IHeadMountedDisplay.h
+++ b/Engine/Source/Runtime/HeadMountedDisplay/Public/IHeadMountedDisplay.h
@@ -70,7 +70,7 @@ public:
     /**
 	 * Calculates the FOV, based on the screen dimensions of the device
 	 */
-	virtual void GetFieldOfView(float& OutHFOVInDegrees, float& OutVFOVInDegrees) const = 0;
+	virtual void	GetFieldOfView(float& OutHFOVInDegrees, float& OutVFOVInDegrees) const = 0;
 
 	/**
 	 * Whether or not the HMD supports positional tracking (either via camera or other means)
@@ -147,6 +147,11 @@ public:
 	 */
     virtual bool Exec( UWorld* InWorld, const TCHAR* Cmd, FOutputDevice& Ar ) = 0;
 
+	/**
+	* Returns true, if HMD allows fullscreen mode.
+	*/
+	virtual bool IsFullScreenAllowed() const { return true; }
+
 	/**
 	 * Saves / loads pre-fullscreen rectangle. Could be used to store saved original window position 
 	 * before switching to fullscreen mode.
@@ -208,6 +213,13 @@ public:
 	 */
 	virtual void DrawDebug(UCanvas* Canvas, EStereoscopicPass StereoPass) {}
 
+	/**
+	* Passing key events to HMD.
+	* If returns 'false' then key will be handled by PlayerController;
+	* otherwise, key won't be handled by the PlayerController.
+	*/
+	virtual bool HandleInputKey(class UPlayerInput*, const struct FKey& Key, enum EInputEvent EventType, float AmountDepressed, bool bGamepad) { return false; }
+
 	/** 
 	 * Additional optional distorion rendering parameters
 	 * @todo:  Once we can move shaders into plugins, remove these!
diff --git a/Engine/Source/Runtime/HeadMountedDisplay/Public/IHeadMountedDisplayModule.h b/Engine/Source/Runtime/HeadMountedDisplay/Public/IHeadMountedDisplayModule.h
--- a/Engine/Source/Runtime/HeadMountedDisplay/Public/IHeadMountedDisplayModule.h
+++ b/Engine/Source/Runtime/HeadMountedDisplay/Public/IHeadMountedDisplayModule.h
@@ -52,4 +52,8 @@ public:
 	 */
 	virtual TSharedPtr< class IHeadMountedDisplay > CreateHeadMountedDisplay() = 0;
 
+	/**
+	* Optionally pre-init the HMD module.
+	*/
+	virtual void PreInit() {}
 };
diff --git a/Engine/Source/Runtime/Launch/Launch.Build.cs b/Engine/Source/Runtime/Launch/Launch.Build.cs
--- a/Engine/Source/Runtime/Launch/Launch.Build.cs
+++ b/Engine/Source/Runtime/Launch/Launch.Build.cs
@@ -35,6 +35,7 @@ public class Launch : ModuleRules
 				"Slate",
 				"SlateCore",
 				"Sockets",
+				"HeadMountedDisplay",
 			}
 		);
         
@@ -49,6 +50,7 @@ public class Launch : ModuleRules
             DynamicallyLoadedModuleNames.AddRange(
                 new string[] {
 			        "SlateRHIRenderer",
+					"HeadMountedDisplay"
 		        }
             );
         }
diff --git a/Engine/Source/Runtime/Launch/Private/LaunchEngineLoop.cpp b/Engine/Source/Runtime/Launch/Private/LaunchEngineLoop.cpp
--- a/Engine/Source/Runtime/Launch/Private/LaunchEngineLoop.cpp
+++ b/Engine/Source/Runtime/Launch/Private/LaunchEngineLoop.cpp
@@ -45,6 +45,7 @@
 	#include "Commandlets/Commandlet.h"
 	#include "EngineService.h"
 	#include "ContentStreaming.h"
+	#include "HeadMountedDisplay.h"
 
 #if !UE_SERVER
 	#include "ISlateRHIRendererModule.h"
@@ -2360,6 +2361,8 @@ bool FEngineLoop::AppInit( )
 		return false;
 	}
 
+	PreInitHMDDevice();
+
 	// Put the command line and config info into the suppression system
 	FLogSuppressionInterface::Get().ProcessConfigAndCommandLine();
 
@@ -2571,4 +2574,24 @@ void FEngineLoop::AppExit( )
 	FInternationalization::TearDown();
 }
 
+void FEngineLoop::PreInitHMDDevice()
+{
+#if WITH_ENGINE
+	if (!GIsEditor)
+	{
+#if 0	//@todo vr: only preinit first valid hmd
+		if (!FParse::Param(FCommandLine::Get(), TEXT("nohmd")) && !FParse::Param(FCommandLine::Get(), TEXT("emulatestereo")))
+		{
+			// Get a list of plugins that implement this feature
+			TArray<IHeadMountedDisplayModule*> HMDImplementations = IModularFeatures::Get().GetModularFeatureImplementations<IHeadMountedDisplayModule>(IHeadMountedDisplayModule::GetModularFeatureName());
+			for (auto HMDModuleIt = HMDImplementations.CreateIterator(); HMDModuleIt; ++HMDModuleIt)
+			{
+				(*HMDModuleIt)->PreInit();
+			}
+		}
+#endif
+	}
+#endif // #if WITH_ENGINE
+}
+
 #undef LOCTEXT_NAMESPACE
diff --git a/Engine/Source/Runtime/Launch/Public/LaunchEngineLoop.h b/Engine/Source/Runtime/Launch/Public/LaunchEngineLoop.h
--- a/Engine/Source/Runtime/Launch/Public/LaunchEngineLoop.h
+++ b/Engine/Source/Runtime/Launch/Public/LaunchEngineLoop.h
@@ -106,6 +106,9 @@ public:
 
 #endif // WITH_ENGINE
 
+	/** Pre-init HMD device (if necessary) */
+	static void PreInitHMDDevice();
+
 public:
 
 	/**
diff --git a/Engine/Source/Runtime/Windows/D3D11RHI/Public/D3D11Viewport.h b/Engine/Source/Runtime/Windows/D3D11RHI/Public/D3D11Viewport.h
--- a/Engine/Source/Runtime/Windows/D3D11RHI/Public/D3D11Viewport.h
+++ b/Engine/Source/Runtime/Windows/D3D11RHI/Public/D3D11Viewport.h
@@ -79,6 +79,9 @@ public:
 	{
 		CustomPresent = InCustomPresent;
 	}
+
+	virtual void* GetNativeWindow(void** AddParam = nullptr) const override { return (void*)&WindowHandle; }
+
 private:
 
 	/** Presents the frame synchronizing with DWM. */
