#ue4 Adding HMD module priority order, which you can specify in your Engine.ini if you need to force a priority on which HMD to use when multiple are present

diff --git a/Engine/Plugins/Experimental/SteamVR/Source/SteamVR/Private/SteamVRHMD.cpp b/Engine/Plugins/Experimental/SteamVR/Source/SteamVR/Private/SteamVRHMD.cpp
--- a/Engine/Plugins/Experimental/SteamVR/Source/SteamVR/Private/SteamVRHMD.cpp
+++ b/Engine/Plugins/Experimental/SteamVR/Source/SteamVR/Private/SteamVRHMD.cpp
@@ -16,6 +16,11 @@ class FSteamVRPlugin : public ISteamVRPlugin
 {
 	/** IHeadMountedDisplayModule implementation */
 	virtual TSharedPtr< class IHeadMountedDisplay > CreateHeadMountedDisplay() override;
+
+	FString GetModulePriorityKeyName() const
+	{
+		return FString(TEXT("SteamVR"));
+	}
 };
 
 IMPLEMENT_MODULE( FSteamVRPlugin, SteamVR )
diff --git a/Engine/Plugins/Runtime/OculusRift/Source/OculusRift/Private/OculusRiftHMD.cpp b/Engine/Plugins/Runtime/OculusRift/Source/OculusRift/Private/OculusRiftHMD.cpp
--- a/Engine/Plugins/Runtime/OculusRift/Source/OculusRift/Private/OculusRiftHMD.cpp
+++ b/Engine/Plugins/Runtime/OculusRift/Source/OculusRift/Private/OculusRiftHMD.cpp
@@ -27,6 +27,11 @@ class FOculusRiftPlugin : public IOculusRiftPlugin
 
 	// Pre-init the HMD module (optional).
 	virtual void PreInit() override;
+
+	FString GetModulePriorityKeyName() const
+	{
+		return FString(TEXT("OculusRift"));
+	}
 };
 
 IMPLEMENT_MODULE( FOculusRiftPlugin, OculusRift )
diff --git a/Engine/Source/Runtime/Engine/Private/UnrealEngine.cpp b/Engine/Source/Runtime/Engine/Private/UnrealEngine.cpp
--- a/Engine/Source/Runtime/Engine/Private/UnrealEngine.cpp
+++ b/Engine/Source/Runtime/Engine/Private/UnrealEngine.cpp
@@ -1730,6 +1730,7 @@ bool UEngine::InitializeHMDDevice()
 		{
 			// Get a list of plugins that implement this feature
 			TArray<IHeadMountedDisplayModule*> HMDImplementations = IModularFeatures::Get().GetModularFeatureImplementations<IHeadMountedDisplayModule>( IHeadMountedDisplayModule::GetModularFeatureName() );
+			HMDImplementations.Sort(FHMDPluginSorter());
 			for( auto HMDModuleIt = HMDImplementations.CreateIterator(); HMDModuleIt && !HMDDevice.IsValid(); ++HMDModuleIt )
 			{
 				HMDDevice = (*HMDModuleIt)->CreateHeadMountedDisplay();
diff --git a/Engine/Source/Runtime/HeadMountedDisplay/Private/HeadMountedDisplayModule.cpp b/Engine/Source/Runtime/HeadMountedDisplay/Private/HeadMountedDisplayModule.cpp
--- a/Engine/Source/Runtime/HeadMountedDisplay/Private/HeadMountedDisplayModule.cpp
+++ b/Engine/Source/Runtime/HeadMountedDisplay/Private/HeadMountedDisplayModule.cpp
@@ -9,6 +9,11 @@ class FHeadMountedDisplayModule : public IHeadMountedDisplayModule
 		TSharedPtr<IHeadMountedDisplay> DummyVal = NULL;
 		return DummyVal;
 	}
+
+	FString GetModulePriorityKeyName() const
+	{
+		return FString(TEXT("Default"));
+	}
 };
 
 IMPLEMENT_MODULE( FHeadMountedDisplayModule, HeadMountedDisplay );
diff --git a/Engine/Source/Runtime/HeadMountedDisplay/Public/IHeadMountedDisplayModule.h b/Engine/Source/Runtime/HeadMountedDisplay/Public/IHeadMountedDisplayModule.h
--- a/Engine/Source/Runtime/HeadMountedDisplay/Public/IHeadMountedDisplayModule.h
+++ b/Engine/Source/Runtime/HeadMountedDisplay/Public/IHeadMountedDisplayModule.h
@@ -6,7 +6,6 @@
 #include "ModuleManager.h"		// For inline LoadModuleChecked()
 #include "Runtime/Core/Public/Features/IModularFeatures.h"
 
-
 /**
  * The public interface of the MotionControlsModule
  */
@@ -19,6 +18,19 @@ public:
 		return HMDFeatureName;
 	}
 
+	/** Returns the priority of this module from INI file configuration */
+	float GetHMDPriority() const
+	{
+		float ModulePriority = 0.f;
+		FString KeyName = GetModulePriorityKeyName();
+		GConfig->GetFloat(TEXT("HMDPluginPriority"), (!KeyName.IsEmpty() ? *KeyName : TEXT("Default")), ModulePriority, GEngineIni);
+
+		return ModulePriority;
+	}
+
+	/** Returns the key into the HMDPluginPriority section of the config file for this module */
+	virtual FString GetModulePriorityKeyName() const = 0;
+	
 	virtual void StartupModule() override
 	{
 		IModularFeatures::Get().RegisterModularFeature( GetModularFeatureName(), this );
@@ -57,3 +69,15 @@ public:
 	*/
 	virtual void PreInit() {}
 };
+
+/** Sorting method for which plugin should be given priority */
+struct FHMDPluginSorter
+{
+	FHMDPluginSorter() {}
+
+	// Sort predicate operator
+	bool operator()(IHeadMountedDisplayModule& LHS, IHeadMountedDisplayModule& RHS) const
+	{
+		return LHS.GetHMDPriority() > RHS.GetHMDPriority();
+	}
+};
\ No newline at end of file
