- Updated GearVR for Note 4 and general cleanup.

diff --git a/Engine/Source/Runtime/HeadMountedDisplay/Public/IHeadMountedDisplay.h b/Engine/Source/Runtime/HeadMountedDisplay/Public/IHeadMountedDisplay.h
--- a/Engine/Source/Runtime/HeadMountedDisplay/Public/IHeadMountedDisplay.h
+++ b/Engine/Source/Runtime/HeadMountedDisplay/Public/IHeadMountedDisplay.h
@@ -17,7 +17,8 @@ namespace EHMDDeviceType
 	{
 		DT_OculusRift,
 		DT_Morpheus,
-		DT_ES2GenericStereoMesh
+		DT_ES2GenericStereoMesh,
+		DT_NoPost						// don't register any post passes
 	};
 }
 
