Fixed copy/paste comment.

diff --git a/Engine/Source/Runtime/HeadMountedDisplay/Public/IHeadMountedDisplayModule.h b/Engine/Source/Runtime/HeadMountedDisplay/Public/IHeadMountedDisplayModule.h
--- a/Engine/Source/Runtime/HeadMountedDisplay/Public/IHeadMountedDisplayModule.h
+++ b/Engine/Source/Runtime/HeadMountedDisplay/Public/IHeadMountedDisplayModule.h
@@ -37,7 +37,7 @@ public:
 	}
 
 	/**
-	 * Singleton-like access to IMovieSceneCore
+	 * Singleton-like access to IHeadMountedDisplayModule
 	 *
 	 * @return Returns IHeadMountedDisplayModule singleton instance, loading the module on demand if needed
 	 */
diff --git a/Engine/Source/Runtime/InputDevice/Public/IInputDeviceModule.h b/Engine/Source/Runtime/InputDevice/Public/IInputDeviceModule.h
--- a/Engine/Source/Runtime/InputDevice/Public/IInputDeviceModule.h
+++ b/Engine/Source/Runtime/InputDevice/Public/IInputDeviceModule.h
@@ -25,7 +25,7 @@ public:
 	}
 
 	/**
-	 * Singleton-like access to IMovieSceneCore
+	 * Singleton-like access to IInputDeviceModule
 	 *
 	 * @return Returns IInputDeviceModule singleton instance, loading the module on demand if needed
 	 */
