#motioncontroller Preventing non-authoritative motion controller components from moving remote clients

diff --git a/Engine/Source/Runtime/HeadMountedDisplay/Private/MotionControllerComponent.cpp b/Engine/Source/Runtime/HeadMountedDisplay/Private/MotionControllerComponent.cpp
--- a/Engine/Source/Runtime/HeadMountedDisplay/Private/MotionControllerComponent.cpp
+++ b/Engine/Source/Runtime/HeadMountedDisplay/Private/MotionControllerComponent.cpp
@@ -22,7 +22,10 @@ void UMotionControllerComponent::TickComponent(float DeltaTime, enum ELevelTick
 	FVector Position = FVector::ZeroVector;
 	FRotator Orientation = FRotator::ZeroRotator;
 
-	if (PlayerIndex != INDEX_NONE)
+	const APlayerController* Actor = Cast<APlayerController>(GetOwner());
+	const bool bHasAuthority = !Actor || Actor->IsLocalPlayerController();
+
+	if ((PlayerIndex != INDEX_NONE) && bHasAuthority)
 	{
 		for (auto MotionController : GEngine->MotionControllerDevices)
 		{
diff --git a/Engine/Source/Runtime/HeadMountedDisplay/Public/MotionControllerComponent.h b/Engine/Source/Runtime/HeadMountedDisplay/Public/MotionControllerComponent.h
--- a/Engine/Source/Runtime/HeadMountedDisplay/Public/MotionControllerComponent.h
+++ b/Engine/Source/Runtime/HeadMountedDisplay/Public/MotionControllerComponent.h
@@ -4,7 +4,7 @@
 #include "IMotionController.h"
 #include "MotionControllerComponent.generated.h"
 
-UCLASS(MinimalAPI, meta = (BlueprintSpawnableComponent), ClassGroup = MotionController)
+UCLASS(MinimalAPI, Blueprintable, meta = (BlueprintSpawnableComponent), ClassGroup = MotionController)
 class UMotionControllerComponent : public UPrimitiveComponent
 {
 	GENERATED_UCLASS_BODY()
