Custom HMD post process mesh implementation.
Instead of rendering a full screen quad for post process passes, render a mesh which represents only the screen area visible in the HMD after distortion.
This cuts the cost of post processing ~10-15% with no visible difference.
Initial implementation for the Vive.

diff --git a/Engine/Plugins/Runtime/Steam/SteamVR/Source/SteamVR/Private/SteamVRHMD.cpp b/Engine/Plugins/Runtime/Steam/SteamVR/Source/SteamVR/Private/SteamVRHMD.cpp
--- a/Engine/Plugins/Runtime/Steam/SteamVR/Source/SteamVR/Private/SteamVRHMD.cpp
+++ b/Engine/Plugins/Runtime/Steam/SteamVR/Source/SteamVR/Private/SteamVRHMD.cpp
@@ -9,6 +9,8 @@
 #include "PostProcess/PostProcessHMD.h"
 #include "Classes/SteamVRFunctionLibrary.h"
 
+#include "ViveVisibleAreaMesh.h"
+
 #if WITH_EDITOR
 #include "Editor/UnrealEd/Classes/Editor/EditorEngine.h"
 #endif
@@ -992,9 +994,11 @@ void FSteamVRHMD::Startup()
 			}
 		}
 
-		// Setup hidden area meshes
+		// Setup meshes
 		HiddenAreaMeshes[0].Build(VRSystem->GetHiddenAreaMesh(vr::Hmd_Eye::Eye_Left));
 		HiddenAreaMeshes[1].Build(VRSystem->GetHiddenAreaMesh(vr::Hmd_Eye::Eye_Right));
+		VisibleAreaMeshes[0].Build(LeftEyePositions, VertexCount);
+		VisibleAreaMeshes[1].Build(RightEyePositions, VertexCount);
 
 #if PLATFORM_WINDOWS
 		if (IsPCPlatform(GMaxRHIShaderPlatform) && !IsOpenGLPlatform(GMaxRHIShaderPlatform))
@@ -1173,4 +1177,59 @@ void FSteamVRHMD::FHiddenAreaMesh::Build(const vr::HiddenAreaMesh_t& Mesh)
 	}
 }
 
+FSteamVRHMD::FVisibleAreaMesh::FVisibleAreaMesh() :
+	pVertices(nullptr),
+	pIndices(nullptr),
+	NumVertices(0),
+	NumIndices(0),
+	NumTriangles(0)
+{}
+
+FSteamVRHMD::FVisibleAreaMesh::~FVisibleAreaMesh()
+{
+	if (pVertices)
+	{
+		delete[] pVertices;
+	}
+
+	if (pIndices)
+	{
+		delete[] pIndices;
+	}
+}
+
+void FSteamVRHMD::FVisibleAreaMesh::Build(const MeshVertex Positions[], uint32 InNumVertices)
+{
+	check(pVertices == nullptr);
+	check(InNumVertices > 2 && InNumVertices % 3 == 0);
+
+	NumVertices = InNumVertices;
+	NumTriangles = NumVertices / 3;
+	NumIndices = NumVertices;
+
+	pVertices = new FFilterVertex[NumVertices];
+	pIndices = new uint16[NumIndices];
+
+	uint32 DataIndex = 0;
+	for (uint32 TriangleIter = 0; TriangleIter < NumTriangles; ++TriangleIter)
+	{
+		for (uint32 VertexIter = 0; VertexIter < 3; ++VertexIter)
+		{
+			const MeshVertex& Position = Positions[DataIndex];
+			FFilterVertex& Vertex = pVertices[DataIndex];
+			
+			Vertex.Position.X = Position.X;
+			Vertex.Position.Y = 1.0f - Position.Y;
+			Vertex.Position.Z = 0.0f;
+			Vertex.Position.W = 1.0f;
+			
+			Vertex.UV.X = Position.X;
+			Vertex.UV.Y = 1.0f - Position.Y;
+			
+			pIndices[DataIndex] = DataIndex;
+			
+			++DataIndex;
+		}
+	}
+}
 #endif //STEAMVR_SUPPORTED_PLATFORMS
diff --git a/Engine/Plugins/Runtime/Steam/SteamVR/Source/SteamVR/Private/SteamVRHMD.h b/Engine/Plugins/Runtime/Steam/SteamVR/Source/SteamVR/Private/SteamVRHMD.h
--- a/Engine/Plugins/Runtime/Steam/SteamVR/Source/SteamVR/Private/SteamVRHMD.h
+++ b/Engine/Plugins/Runtime/Steam/SteamVR/Source/SteamVR/Private/SteamVRHMD.h
@@ -85,7 +85,10 @@ public:
 	virtual FQuat GetBaseOrientation() const override;
 
 	virtual bool HasHiddenAreaMask() const override { return HiddenAreaMeshes[0].IsValid() && HiddenAreaMeshes[1].IsValid(); }
-	virtual void DrawHiddenAreaMaskView_RenderThread(FRHICommandList& RHICmdList, const FViewInfo& View) const override;
+	virtual void DrawHiddenAreaMaskView_RenderThread(FRHICommandList& RHICmdList, EStereoscopicPass StereoPass) const override;
+
+	virtual bool HasPostProcessMesh() const override { return VisibleAreaMeshes[0].IsValid() && VisibleAreaMeshes[1].IsValid(); }
+	virtual void DrawPostProcessMesh_RenderThread(FRHICommandList& RHICmdList, EStereoscopicPass StereoPass) const override;
 
 	virtual void DrawDistortionMesh_RenderThread(struct FRenderingCompositePassContext& Context, const FIntPoint& TextureSize) override;
 
@@ -316,7 +319,27 @@ private:
 		unsigned  NumTriangles;
 	};
 
+	struct FVisibleAreaMesh
+	{
+		FVisibleAreaMesh();
+		~FVisibleAreaMesh();
+
+		bool IsValid() const
+		{
+			return NumTriangles > 0;
+		}
+
+		void Build(const struct MeshVertex Positions[], uint32 InNumVertices);
+
+		FFilterVertex* pVertices;
+		uint16*   pIndices;
+		unsigned  NumVertices;
+		unsigned  NumIndices;
+		unsigned  NumTriangles;
+	};
+
 	FHiddenAreaMesh HiddenAreaMeshes[2];
+	FVisibleAreaMesh VisibleAreaMeshes[2];
 
 	/** Chaperone Support */
 	struct FChaperoneBounds
diff --git a/Engine/Plugins/Runtime/Steam/SteamVR/Source/SteamVR/Private/SteamVRRender.cpp b/Engine/Plugins/Runtime/Steam/SteamVR/Source/SteamVR/Private/SteamVRRender.cpp
--- a/Engine/Plugins/Runtime/Steam/SteamVR/Source/SteamVR/Private/SteamVRRender.cpp
+++ b/Engine/Plugins/Runtime/Steam/SteamVR/Source/SteamVR/Private/SteamVRRender.cpp
@@ -75,12 +75,12 @@ void FSteamVRHMD::RenderTexture_RenderThread(FRHICommandListImmediate& RHICmdLis
 	}
 }
 
-void FSteamVRHMD::DrawHiddenAreaMaskView_RenderThread(FRHICommandList& RHICmdList, const FViewInfo& View) const
+void FSteamVRHMD::DrawHiddenAreaMaskView_RenderThread(FRHICommandList& RHICmdList, EStereoscopicPass StereoPass) const
 {
 	check(IsInRenderingThread());
-	check(View.StereoPass != eSSP_FULL);
+	check(StereoPass != eSSP_FULL);
 
-	const uint32 HiddenMeshIndex = (View.StereoPass == eSSP_LEFT_EYE) ? 0 : 1;
+	const uint32 HiddenMeshIndex = (StereoPass == eSSP_LEFT_EYE) ? 0 : 1;
 	const FHiddenAreaMesh& Mesh = HiddenAreaMeshes[HiddenMeshIndex];
 	check(Mesh.IsValid());
 
@@ -97,6 +97,28 @@ void FSteamVRHMD::DrawHiddenAreaMaskView_RenderThread(FRHICommandList& RHICmdLis
 	);
 }
 
+void FSteamVRHMD::DrawPostProcessMesh_RenderThread(FRHICommandList& RHICmdList, EStereoscopicPass StereoPass) const
+{
+	check(IsInRenderingThread());
+	check(StereoPass != eSSP_FULL);
+
+	const uint32 HiddenMeshIndex = (StereoPass == eSSP_LEFT_EYE) ? 0 : 1;
+	const FVisibleAreaMesh& Mesh = VisibleAreaMeshes[HiddenMeshIndex];
+	check(Mesh.IsValid());
+
+	DrawIndexedPrimitiveUP(
+		RHICmdList,
+		PT_TriangleList,
+		0,
+		Mesh.NumVertices,
+		Mesh.NumTriangles,
+		Mesh.pIndices,
+		sizeof(Mesh.pIndices[0]),
+		Mesh.pVertices,
+		sizeof(Mesh.pVertices[0])
+		);
+}
+
 #if PLATFORM_WINDOWS
 
 FSteamVRHMD::D3D11Bridge::D3D11Bridge(FSteamVRHMD* plugin):
diff --git a/Engine/Plugins/Runtime/Steam/SteamVR/Source/SteamVR/Private/ViveVisibleAreaMesh.h b/Engine/Plugins/Runtime/Steam/SteamVR/Source/SteamVR/Private/ViveVisibleAreaMesh.h
new file mode 100644
--- /dev/null
+++ b/Engine/Plugins/Runtime/Steam/SteamVR/Source/SteamVR/Private/ViveVisibleAreaMesh.h
@@ -0,0 +1,164 @@
+// Copyright 1998-2015 Epic Games, Inc. All Rights Reserved.
+
+#pragma once
+
+struct MeshVertex
+{
+	MeshVertex(float InX, float InY) :
+		X(InX),
+		Y(InY)
+	{}
+
+	float X;
+	float Y;
+};
+
+// These are the inverse of the hidden area meshes Valve provides us for the Vive.
+// i.e. a mesh that covers all pixels which *are* inputs to the distortion pass.
+
+static const uint32 VertexCount = 69;
+static const MeshVertex LeftEyePositions[] =
+{
+	MeshVertex(1.0f, 0.1999999881f),
+	MeshVertex(1.0f, 0.1949999332f),
+	MeshVertex(0.9099999552f, 0.1099999547f),
+	MeshVertex(0.9159999971f, 0.3639999628f),
+	MeshVertex(1.0f, 0.1999999881f),
+	MeshVertex(0.9099999552f, 0.1099999547f),
+	MeshVertex(0.9099999552f, 0.1099999547f),
+	MeshVertex(0.9159999971f, 0.3639999628f),
+	MeshVertex(0.8100000505f, 0.04499995708f),
+	MeshVertex(0.8100000505f, 0.04499995708f),
+	MeshVertex(0.9159999971f, 0.3639999628f),
+	MeshVertex(0.6799999361f, 0.0f),
+	MeshVertex(0.6799999361f, 0.0f),
+	MeshVertex(0.9159999971f, 0.3639999628f),
+	MeshVertex(0.3699999933f, 0.0f),
+	MeshVertex(0.8949999695f, 0.4949999452f),
+	MeshVertex(0.9159999971f, 0.3639999628f),
+	MeshVertex(0.3699999933f, 0.0f),
+	MeshVertex(0.3699999933f, 0.0f),
+	MeshVertex(0.8949999695f, 0.4949999452f),
+	MeshVertex(0.2399999981f, 0.04999995232f),
+	MeshVertex(0.2399999981f, 0.04999995232f),
+	MeshVertex(0.8949999695f, 0.4949999452f),
+	MeshVertex(0.1320000176f, 0.1229999661f),
+	MeshVertex(0.1320000176f, 0.1229999661f),
+	MeshVertex(0.8949999695f, 0.4949999452f),
+	MeshVertex(0.05999999094f, 0.1999999881f),
+	MeshVertex(0.05999999094f, 0.1999999881f),
+	MeshVertex(0.8949999695f, 0.4949999452f),
+	MeshVertex(0.0f, 0.2949999571f),
+	MeshVertex(0.0f, 0.2949999571f),
+	MeshVertex(0.8949999695f, 0.4949999452f),
+	MeshVertex(0.0f, 0.6999999285f),
+	MeshVertex(0.0f, 0.6999999285f),
+	MeshVertex(0.8949999695f, 0.4949999452f),
+	MeshVertex(0.04000001001f, 0.7680000067f),
+	MeshVertex(0.04000001001f, 0.7680000067f),
+	MeshVertex(0.8949999695f, 0.4949999452f),
+	MeshVertex(0.09000002193f, 0.8329999447f),
+	MeshVertex(0.09000002193f, 0.8329999447f),
+	MeshVertex(0.8949999695f, 0.4949999452f),
+	MeshVertex(0.1899999862f, 0.9149999619f),
+	MeshVertex(0.1899999862f, 0.9149999619f),
+	MeshVertex(0.8949999695f, 0.4949999452f),
+	MeshVertex(0.2710000162f, 0.9599999189f),
+	MeshVertex(0.2710000162f, 0.9599999189f),
+	MeshVertex(0.8949999695f, 0.4949999452f),
+	MeshVertex(0.388f, 1.0f),
+	MeshVertex(0.388f, 1.0f),
+	MeshVertex(0.8949999695f, 0.4950000048f),
+	MeshVertex(0.6560000067f, 1.0f),
+	MeshVertex(0.9129999762f, 0.628000021f),
+	MeshVertex(0.8949999695f, 0.4949999452f),
+	MeshVertex(0.6560000067f, 1.0f),
+	MeshVertex(0.6560000067f, 1.0f),
+	MeshVertex(0.9129999762f, 0.628000021f),
+	MeshVertex(0.7399999981f, 0.9749999046f),
+	MeshVertex(0.7399999981f, 0.9749999046f),
+	MeshVertex(0.9129999762f, 0.628000021f),
+	MeshVertex(0.8129999523f, 0.9449999332f),
+	MeshVertex(0.8129999523f, 0.9449999332f),
+	MeshVertex(0.9129999762f, 0.628000021f),
+	MeshVertex(0.9299999361f, 0.8679999113f),
+	MeshVertex(1.0f, 0.7999999523f),
+	MeshVertex(0.9299999361f, 0.8679999113f),
+	MeshVertex(1.0f, 0.7879999876f),
+	MeshVertex(1.0f, 0.7879999876f),
+	MeshVertex(0.9129999762f, 0.628000021f),
+	MeshVertex(0.9299999361f, 0.8679999113f),
+};
+
+static const MeshVertex RightEyePositions[] =
+{
+	MeshVertex(0.09000000358f, 0.1099999917f),
+	MeshVertex(0.0f, 0.195f),
+	MeshVertex(0.0f, 0.1999999952f),
+	MeshVertex(0.0f, 0.1999999952f),
+	MeshVertex(0.09000000358f, 0.1099999917f),
+	MeshVertex(0.08399999142f, 0.3639999998f),
+	MeshVertex(0.1899999976f, 0.04499999404f),
+	MeshVertex(0.09000000358f, 0.1099999917f),
+	MeshVertex(0.08399999142f, 0.3639999998f),
+	MeshVertex(0.3199999928f, 0.0f),
+	MeshVertex(0.1899999976f, 0.04499999404f),
+	MeshVertex(0.08399999142f, 0.3639999998f),
+	MeshVertex(0.6299999952f, 0.0f),
+	MeshVertex(0.3199999928f, 0.0f),
+	MeshVertex(0.08399999142f, 0.3639999998f),
+	MeshVertex(0.08399999142f, 0.3639999998f),
+	MeshVertex(0.6299999952f, 0.0f),
+	MeshVertex(0.1049999893f, 0.4949999821f),
+	MeshVertex(0.7599999905f, 0.04999998927f),
+	MeshVertex(0.6299999952f, 0.0f),
+	MeshVertex(0.1049999893f, 0.4949999821f),
+	MeshVertex(0.8680000305f, 0.1230000031f),
+	MeshVertex(0.7599999905f, 0.04999998927f),
+	MeshVertex(0.1049999893f, 0.4949999821f),
+	MeshVertex(0.9400000572f, 0.1999999952f),
+	MeshVertex(0.8680000305f, 0.1230000031f),
+	MeshVertex(0.1049999893f, 0.4949999821f),
+	MeshVertex(1.0f, 0.294999994f),
+	MeshVertex(0.9400000572f, 0.1999999952f),
+	MeshVertex(0.1049999893f, 0.4949999821f),
+	MeshVertex(1.0f, 0.700000025f),
+	MeshVertex(1.0f, 0.294999994f),
+	MeshVertex(0.1049999893f, 0.4949999821f),
+	MeshVertex(0.9600000381f, 0.767999984f),
+	MeshVertex(1.0f, 0.700000025f),
+	MeshVertex(0.1049999893f, 0.4949999821f),
+	MeshVertex(0.9099999666f, 0.8330000412f),
+	MeshVertex(0.9600000381f, 0.767999984f),
+	MeshVertex(0.1049999893f, 0.4949999821f),
+	MeshVertex(0.8099999428f, 0.9150000584f),
+	MeshVertex(0.9099999666f, 0.8330000412f),
+	MeshVertex(0.1049999893f, 0.4949999821f),
+	MeshVertex(0.7289999723f, 0.9600000155f),
+	MeshVertex(0.8099999428f, 0.9150000584f),
+	MeshVertex(0.1049999893f, 0.4949999821f),
+	MeshVertex(0.6119999886f, 1.0f),
+	MeshVertex(0.7289999723f, 0.9600000155f),
+	MeshVertex(0.1049999893f, 0.4949999821f),
+	MeshVertex(0.3439999819f, 1.0f),
+	MeshVertex(0.6119999886f, 1.0f),
+	MeshVertex(0.1049999893f, 0.4949999821f),
+	MeshVertex(0.1049999893f, 0.4949999821f),
+	MeshVertex(0.3439999819f, 1.0f),
+	MeshVertex(0.0870000124f, 0.6279999983f),
+	MeshVertex(0.2599999905f, 0.9750000012f),
+	MeshVertex(0.3439999819f, 1.0f),
+	MeshVertex(0.0870000124f, 0.6279999983f),
+	MeshVertex(0.1870000064f, 0.9450000298f),
+	MeshVertex(0.2599999905f, 0.9750000012f),
+	MeshVertex(0.0870000124f, 0.6279999983f),
+	MeshVertex(0.06999999285f, 0.8680000079f),
+	MeshVertex(0.1870000064f, 0.9450000298f),
+	MeshVertex(0.0870000124f, 0.6279999983f),
+	MeshVertex(0.0f, 0.8000000489f),
+	MeshVertex(0.0f, 0.787999965f),
+	MeshVertex(0.06999999285f, 0.8680000079f),
+	MeshVertex(0.0870000124f, 0.6279999983f),
+	MeshVertex(0.06999999285f, 0.8680000079f),
+	MeshVertex(0.0f, 0.787999965f),
+};
diff --git a/Engine/Source/Runtime/HeadMountedDisplay/Public/IHeadMountedDisplay.h b/Engine/Source/Runtime/HeadMountedDisplay/Public/IHeadMountedDisplay.h
--- a/Engine/Source/Runtime/HeadMountedDisplay/Public/IHeadMountedDisplay.h
+++ b/Engine/Source/Runtime/HeadMountedDisplay/Public/IHeadMountedDisplay.h
@@ -49,7 +49,7 @@ public:
 	virtual void EnableHMD(bool bEnable = true) = 0;
 
 	/**
-	 * Returns the family of HMD device implemented 
+	 * Returns the family of HMD device implemented
 	 */
 	virtual EHMDDeviceType::Type GetHMDDeviceType() const = 0;
 
@@ -73,12 +73,12 @@ public:
 
 	};
 
-    /**
-     * Get the name or id of the display to output for this HMD. 
-     */
+	/**
+	 * Get the name or id of the display to output for this HMD.
+	 */
 	virtual bool	GetHMDMonitorInfo(MonitorInfo&) = 0;
-	
-    /**
+
+	/**
 	 * Calculates the FOV, based on the screen dimensions of the device. Original FOV is passed as params.
 	 */
 	virtual void	GetFieldOfView(float& InOutHFOVInDegrees, float& InOutVFOVInDegrees) const = 0;
@@ -96,7 +96,7 @@ public:
 	/**
 	 * If the HMD supports positional tracking via a camera, this returns the frustum properties (all in game-world space) of the tracking camera.
 	 */
-	virtual void	GetPositionalTrackingCameraProperties(FVector& OutOrigin, FQuat& OutOrientation, float& OutHFOV, float& OutVFOV, float& OutCameraDistance, float& OutNearPlane, float& OutFarPlane) const = 0;	
+	virtual void	GetPositionalTrackingCameraProperties(FVector& OutOrigin, FQuat& OutOrientation, float& OutHFOV, float& OutVFOV, float& OutCameraDistance, float& OutNearPlane, float& OutFarPlane) const = 0;
 
 	/**
 	 * Accessors to modify the interpupillary distance (meters)
@@ -104,21 +104,21 @@ public:
 	virtual void	SetInterpupillaryDistance(float NewInterpupillaryDistance) = 0;
 	virtual float	GetInterpupillaryDistance() const = 0;
 
-    /**
-     * Get the current orientation and position reported by the HMD.
-     */
-    virtual void GetCurrentOrientationAndPosition(FQuat& CurrentOrientation, FVector& CurrentPosition) = 0;
+	/**
+	 * Get the current orientation and position reported by the HMD.
+	 */
+	virtual void GetCurrentOrientationAndPosition(FQuat& CurrentOrientation, FVector& CurrentPosition) = 0;
 
-    /**
-     * Get the ISceneViewExtension for this HMD, or none.
-     */
-    virtual TSharedPtr<class ISceneViewExtension, ESPMode::ThreadSafe> GetViewExtension() = 0;
+	/**
+	 * Get the ISceneViewExtension for this HMD, or none.
+	 */
+	virtual TSharedPtr<class ISceneViewExtension, ESPMode::ThreadSafe> GetViewExtension() = 0;
 
 	/**
-     * Apply the orientation of the headset to the PC's rotation.
-     * If this is not done then the PC will face differently than the camera,
-     * which might be good (depending on the game).
-     */
+	 * Apply the orientation of the headset to the PC's rotation.
+	 * If this is not done then the PC will face differently than the camera,
+	 * which might be good (depending on the game).
+	 */
 	virtual void ApplyHmdRotation(class APlayerController* PC, FRotator& ViewRotation) = 0;
 
 	/**
@@ -127,7 +127,7 @@ public:
 	 */
 	virtual void UpdatePlayerCameraRotation(class APlayerCameraManager* Camera, struct FMinimalViewInfo& POV) = 0;
 
-  	/**
+	/**
 	 * Gets the scaling factor, applied to the post process warping effect
 	 */
 	virtual float GetDistortionScalingFactor() const { return 0; }
@@ -148,7 +148,7 @@ public:
 	virtual bool IsChromaAbCorrectionEnabled() const = 0;
 
 	/**
-	 * Gets the chromatic aberration correction shader values for the device. 
+	 * Gets the chromatic aberration correction shader values for the device.
 	 * Returns 'false' if chromatic aberration correction is off.
 	 */
 	virtual bool GetChromaAbCorrectionValues(FVector4& K) const  { return false; }
@@ -156,7 +156,7 @@ public:
 	/**
 	 * Exec handler to allow console commands to be passed through to the HMD for debugging
 	 */
-    virtual bool Exec( UWorld* InWorld, const TCHAR* Cmd, FOutputDevice& Ar ) = 0;
+	virtual bool Exec(UWorld* InWorld, const TCHAR* Cmd, FOutputDevice& Ar) = 0;
 
 	/**
 	 * Returns true, if HMD allows fullscreen mode.
@@ -164,23 +164,23 @@ public:
 	virtual bool IsFullscreenAllowed() { return true; }
 
 	/**
-	 * Saves / loads pre-fullscreen rectangle. Could be used to store saved original window position 
+	 * Saves / loads pre-fullscreen rectangle. Could be used to store saved original window position
 	 * before switching to fullscreen mode.
 	 */
 	virtual void PushPreFullScreenRect(const FSlateRect& InPreFullScreenRect);
 	virtual void PopPreFullScreenRect(FSlateRect& OutPreFullScreenRect);
 
 	/**
-	 * A callback that is called when screen mode is changed (fullscreen <-> window). 
+	 * A callback that is called when screen mode is changed (fullscreen <-> window).
 	 */
 	virtual void OnScreenModeChange(EWindowMode::Type WindowMode) = 0;
 
 	/** Returns true if positional tracking enabled and working. */
 	virtual bool IsPositionalTrackingEnabled() const = 0;
 
-	/** 
+	/**
 	 * Tries to enable positional tracking.
-	 * Returns the actual status of positional tracking. 
+	 * Returns the actual status of positional tracking.
 	 */
 	virtual bool EnablePositionalTracking(bool enable) = 0;
 
@@ -202,29 +202,29 @@ public:
 	 */
 	virtual void EnableLowPersistenceMode(bool Enable = true) = 0;
 
-	/** 
+	/**
 	 * Resets orientation by setting roll and pitch to 0, assuming that current yaw is forward direction and assuming
-	 * current position as a 'zero-point' (for positional tracking). 
+	 * current position as a 'zero-point' (for positional tracking).
 	 *
 	 * @param Yaw				(in) the desired yaw to be set after orientation reset.
 	 */
 	virtual void ResetOrientationAndPosition(float Yaw = 0.f) = 0;
 
-	/** 
-	 * Resets orientation by setting roll and pitch to 0, assuming that current yaw is forward direction. Position is not changed. 
+	/**
+	 * Resets orientation by setting roll and pitch to 0, assuming that current yaw is forward direction. Position is not changed.
 	 *
 	 * @param Yaw				(in) the desired yaw to be set after orientation reset.
 	 */
 	virtual void ResetOrientation(float Yaw = 0.f) {}
 
-	/** 
-	 * Resets position, assuming current position as a 'zero-point'. 
+	/**
+	 * Resets position, assuming current position as a 'zero-point'.
 	 */
 	virtual void ResetPosition() {}
 
-	/** 
-	 * Sets base orientation by setting yaw, pitch, roll, assuming that this is forward direction. 
-	 * Position is not changed. 
+	/**
+	 * Sets base orientation by setting yaw, pitch, roll, assuming that this is forward direction.
+	 * Position is not changed.
 	 *
 	 * @param BaseRot			(in) the desired orientation to be treated as a base orientation.
 	 */
@@ -235,9 +235,9 @@ public:
 	 */
 	virtual FRotator GetBaseRotation() const { return FRotator::ZeroRotator; }
 
-	/** 
-	 * Sets base orientation, assuming that this is forward direction. 
-	 * Position is not changed. 
+	/**
+	 * Sets base orientation, assuming that this is forward direction.
+	 * Position is not changed.
 	 *
 	 * @param BaseOrient		(in) the desired orientation to be treated as a base orientation.
 	 */
@@ -265,11 +265,22 @@ public:
 	*/
 	virtual bool HasHiddenAreaMask() const { return false; }
 
+	/**
+	* @return true if a post process mesh is available for the device.
+	*/
+	virtual bool HasPostProcessMesh() const { return false; }
+
 	/**
 	* Optional method to draw a view's hidden area mask where supported.
 	* This can be used to avoid rendering pixels which are not included as input into the final distortion pass.
 	*/
-	virtual void DrawHiddenAreaMaskView_RenderThread(class FRHICommandList& RHICmdList, const FViewInfo& View) const {};
+	virtual void DrawHiddenAreaMaskView_RenderThread(class FRHICommandList& RHICmdList, EStereoscopicPass StereoPass) const {};
+
+	/**
+	* Optional method to draw a view's post process mesh where supported.
+	* This can be used instead of a full screen quad to avoid rendering pixels which are not included as input into the final distortion pass.
+	*/
+	virtual void DrawPostProcessMesh_RenderThread(class FRHICommandList& RHICmdList, EStereoscopicPass StereoPass) const {};
 
 	virtual void DrawDistortionMesh_RenderThread(struct FRenderingCompositePassContext& Context, const FIntPoint& TextureSize) {}
 
diff --git a/Engine/Source/Runtime/Renderer/Private/CompositionLighting/PostProcessAmbient.cpp b/Engine/Source/Runtime/Renderer/Private/CompositionLighting/PostProcessAmbient.cpp
--- a/Engine/Source/Runtime/Renderer/Private/CompositionLighting/PostProcessAmbient.cpp
+++ b/Engine/Source/Runtime/Renderer/Private/CompositionLighting/PostProcessAmbient.cpp
@@ -92,8 +92,7 @@ void FRCPassPostProcessAmbient::Render(FRenderingCompositePassContext& Context)
 
 		PixelShader->SetParameters(Context.RHICmdList, Context, Context.View.FinalPostProcessSettings.ContributingCubemaps[i]);
 
-		// Draw a quad mapping scene color to the view's render target
-		DrawRectangle(
+		DrawPostProcessPass(
 			Context.RHICmdList,
 			0, 0,
 			View.ViewRect.Width(), View.ViewRect.Height(),
@@ -102,6 +101,8 @@ void FRCPassPostProcessAmbient::Render(FRenderingCompositePassContext& Context)
 			View.ViewRect.Size(),
 			FSceneRenderTargets::Get(Context.RHICmdList).GetBufferSizeXY(),
 			*VertexShader,
+			View.StereoPass, 
+			Context.HasHmdMesh(),
 			EDRF_UseTriangleOptimization);
 	}
 }
diff --git a/Engine/Source/Runtime/Renderer/Private/CompositionLighting/PostProcessAmbientOcclusion.cpp b/Engine/Source/Runtime/Renderer/Private/CompositionLighting/PostProcessAmbientOcclusion.cpp
--- a/Engine/Source/Runtime/Renderer/Private/CompositionLighting/PostProcessAmbientOcclusion.cpp
+++ b/Engine/Source/Runtime/Renderer/Private/CompositionLighting/PostProcessAmbientOcclusion.cpp
@@ -343,16 +343,17 @@ void FRCPassPostProcessAmbientOcclusionSetup::Process(FRenderingCompositePassCon
 		VertexShader = SetShaderSetupTempl<0>(Context);
 	}
 
-	// Draw a quad mapping scene color to the view's render target
-	DrawRectangle( 
+	DrawPostProcessPass(
 		Context.RHICmdList,
 		0, 0,
 		DestRect.Width(), DestRect.Height(),
-			SrcRect.Min.X, SrcRect.Min.Y, 
+		SrcRect.Min.X, SrcRect.Min.Y,
 		SrcRect.Width(), SrcRect.Height(),
 		DestRect.Size(),
 		FSceneRenderTargets::Get(Context.RHICmdList).GetBufferSizeXY(),
 		VertexShader,
+		View.StereoPass,
+		Context.HasHmdMesh(),
 		EDRF_UseTriangleOptimization);
 
 	Context.RHICmdList.CopyToResolveTarget(DestRenderTarget.TargetableTexture, DestRenderTarget.ShaderResourceTexture, false, FResolveParams());
@@ -492,8 +493,7 @@ void FRCPassPostProcessAmbientOcclusion::Process(FRenderingCompositePassContext&
 
 	TShaderMapRef<FPostProcessVS> VertexShader(Context.GetShaderMap());
 
-	// Draw a quad mapping scene color to the view's render target
-	DrawRectangle( 
+	DrawPostProcessPass(
 		Context.RHICmdList,
 		0, 0,
 		ViewRect.Width(), ViewRect.Height(),
@@ -502,9 +502,10 @@ void FRCPassPostProcessAmbientOcclusion::Process(FRenderingCompositePassContext&
 		ViewRect.Size(),
 		TexSize,
 		*VertexShader,
+		View.StereoPass,
+		Context.HasHmdMesh(),
 		EDRF_UseTriangleOptimization);
-
-
+	
 	Context.RHICmdList.CopyToResolveTarget(DestRenderTarget->TargetableTexture, DestRenderTarget->ShaderResourceTexture, false, FResolveParams());
 }
 
@@ -563,16 +564,17 @@ void FRCPassPostProcessBasePassAO::Process(FRenderingCompositePassContext& Conte
 	VertexShader->SetParameters(Context);
 	PixelShader->SetParameters(Context, SceneContext.GetBufferSizeXY());
 
-	// Draw a quad mapping scene color to the view's render target
-	DrawRectangle( 
+	DrawPostProcessPass(
 		Context.RHICmdList,
 		0, 0,
 		View.ViewRect.Width(), View.ViewRect.Height(),
-		View.ViewRect.Min.X, View.ViewRect.Min.Y, 
+		View.ViewRect.Min.X, View.ViewRect.Min.Y,
 		View.ViewRect.Width(), View.ViewRect.Height(),
 		View.ViewRect.Size(),
 		SceneContext.GetBufferSizeXY(),
 		*VertexShader,
+		View.StereoPass, 
+		Context.HasHmdMesh(),
 		EDRF_UseTriangleOptimization);
 
 	Context.RHICmdList.CopyToResolveTarget(DestRenderTarget.TargetableTexture, DestRenderTarget.ShaderResourceTexture, false, FResolveParams());
diff --git a/Engine/Source/Runtime/Renderer/Private/CompositionLighting/PostProcessDeferredDecals.cpp b/Engine/Source/Runtime/Renderer/Private/CompositionLighting/PostProcessDeferredDecals.cpp
--- a/Engine/Source/Runtime/Renderer/Private/CompositionLighting/PostProcessDeferredDecals.cpp
+++ b/Engine/Source/Runtime/Renderer/Private/CompositionLighting/PostProcessDeferredDecals.cpp
@@ -296,7 +296,7 @@ IMPLEMENT_SHADER_TYPE(,FStencilDecalMaskPS,TEXT("DeferredDecal"),TEXT("StencilDe
 FGlobalBoundShaderState StencilDecalMaskBoundShaderState;
 
 /** Draws a full view quad that sets stencil to 1 anywhere that decals should not be projected. */
-void StencilDecalMask(FRHICommandList& RHICmdList, const FViewInfo& View)
+void StencilDecalMask(FRHICommandList& RHICmdList, const FViewInfo& View, bool bUseHmdMesh)
 {
 	SCOPED_DRAW_EVENT(RHICmdList, StencilDecalMask);
 	FSceneRenderTargets& SceneContext = FSceneRenderTargets::Get(RHICmdList);
@@ -317,15 +317,17 @@ void StencilDecalMask(FRHICommandList& RHICmdList, const FViewInfo& View)
 
 	PixelShader->SetParameters(RHICmdList, View);
 
-	DrawRectangle( 
+	DrawPostProcessPass(
 		RHICmdList,
-		0, 0, 
+		0, 0,
 		View.ViewRect.Width(), View.ViewRect.Height(),
-		View.ViewRect.Min.X, View.ViewRect.Min.Y, 
+		View.ViewRect.Min.X, View.ViewRect.Min.Y,
 		View.ViewRect.Width(), View.ViewRect.Height(),
 		FIntPoint(View.ViewRect.Width(), View.ViewRect.Height()),
 		SceneContext.GetBufferSizeXY(),
 		*ScreenVertexShader,
+		View.StereoPass,
+		bUseHmdMesh,
 		EDRF_UseTriangleOptimization);
 }
 
@@ -651,7 +653,7 @@ void FRCPassPostProcessDeferredDecals::Process(FRenderingCompositePassContext& C
 			// Setup a stencil mask to prevent certain pixels from receiving deferred decals
 			if (bStencilDecalsInThisStage)
 			{
-				StencilDecalMask(RHICmdList, View);
+				StencilDecalMask(RHICmdList, View, Context.HasHmdMesh());
 			}
 
 			// optimization to have less state changes
diff --git a/Engine/Source/Runtime/Renderer/Private/CompositionLighting/PostProcessLpvIndirect.cpp b/Engine/Source/Runtime/Renderer/Private/CompositionLighting/PostProcessLpvIndirect.cpp
--- a/Engine/Source/Runtime/Renderer/Private/CompositionLighting/PostProcessLpvIndirect.cpp
+++ b/Engine/Source/Runtime/Renderer/Private/CompositionLighting/PostProcessLpvIndirect.cpp
@@ -322,16 +322,17 @@ void FRCPassPostProcessLpvIndirect::Process(FRenderingCompositePassContext& Cont
 	{
 		SCOPED_DRAW_EVENT(Context.RHICmdList, PostProcessLpvIndirect );
 
-	    // Draw a quad mapping scene color to the view's render target
-	    DrawRectangle( 
-		    Context.RHICmdList,
-		    0, 0,
-		    View.ViewRect.Width(), View.ViewRect.Height(),
-		    View.ViewRect.Min.X, View.ViewRect.Min.Y, 
-		    View.ViewRect.Width(), View.ViewRect.Height(),
-		    View.ViewRect.Size(),
-		    SceneContext.GetBufferSizeXY(),
-		    *VertexShader);
+		DrawPostProcessPass(
+			Context.RHICmdList,
+			0, 0,
+			View.ViewRect.Width(), View.ViewRect.Height(),
+			View.ViewRect.Min.X, View.ViewRect.Min.Y,
+			View.ViewRect.Width(), View.ViewRect.Height(),
+			View.ViewRect.Size(),
+			SceneContext.GetBufferSizeXY(),
+			*VertexShader,
+			View.StereoPass,
+			Context.HasHmdMesh());
 
 		Context.RHICmdList.CopyToResolveTarget(DestColorRenderTarget.TargetableTexture, DestColorRenderTarget.ShaderResourceTexture, false, FResolveParams());
 		if(bApplySeparateSpecularRT)
@@ -395,16 +396,17 @@ void FRCPassPostProcessLpvIndirect::DoDirectionalOcclusionPass(FRenderingComposi
 
 	PixelShader->SetParameters( Lpv->GetAOVolumeTextureSRV(), LpvReadUniformBuffer, Context );
 
-	// Draw a quad mapping scene color to the view's render target
-	DrawRectangle( 
+	DrawPostProcessPass(
 		Context.RHICmdList,
 		0, 0,
 		View.ViewRect.Width(), View.ViewRect.Height(),
-		View.ViewRect.Min.X, View.ViewRect.Min.Y, 
+		View.ViewRect.Min.X, View.ViewRect.Min.Y,
 		View.ViewRect.Width(), View.ViewRect.Height(),
 		View.ViewRect.Size(),
 		SceneContext.GetBufferSizeXY(),
-		*VertexShader);
+		*VertexShader, 
+		View.StereoPass, 
+		Context.HasHmdMesh());
 }
 
 FPooledRenderTargetDesc FRCPassPostProcessLpvIndirect::ComputeOutputDesc(EPassOutputId InPassOutputId) const
diff --git a/Engine/Source/Runtime/Renderer/Private/CompositionLighting/PostProcessPassThrough.cpp b/Engine/Source/Runtime/Renderer/Private/CompositionLighting/PostProcessPassThrough.cpp
--- a/Engine/Source/Runtime/Renderer/Private/CompositionLighting/PostProcessPassThrough.cpp
+++ b/Engine/Source/Runtime/Renderer/Private/CompositionLighting/PostProcessPassThrough.cpp
@@ -127,8 +127,7 @@ void FRCPassPostProcessPassThrough::Process(FRenderingCompositePassContext& Cont
 	VertexShader->SetParameters(Context);
 	PixelShader->SetParameters(Context);
 
-	// Draw a quad mapping scene color to the view's render target
-	DrawRectangle(
+	DrawPostProcessPass(
 		Context.RHICmdList,
 		DestRect.Min.X, DestRect.Min.Y,
 		DestRect.Width(), DestRect.Height(),
@@ -137,9 +136,10 @@ void FRCPassPostProcessPassThrough::Process(FRenderingCompositePassContext& Cont
 		DestSize,
 		SrcSize,
 		*VertexShader,
+		View.StereoPass,
+		Context.HasHmdMesh(),
 		EDRF_UseTriangleOptimization);
 
-
 	Context.RHICmdList.CopyToResolveTarget(DestRenderTarget.TargetableTexture, DestRenderTarget.ShaderResourceTexture, false, FResolveParams());
 }
 
@@ -196,8 +196,7 @@ void CopyOverOtherViewportsIfNeeded(FRenderingCompositePassContext& Context, con
 			{
 				FIntRect Rect = LocalView.ViewRect;
 
-				DrawRectangle(
-					Context.RHICmdList,
+				DrawPostProcessPass(Context.RHICmdList,
 					Rect.Min.X, Rect.Min.Y,
 					Rect.Width(), Rect.Height(),
 					Rect.Min.X, Rect.Min.Y,
@@ -205,6 +204,8 @@ void CopyOverOtherViewportsIfNeeded(FRenderingCompositePassContext& Context, con
 					Size,
 					Size,
 					*VertexShader,
+					View.StereoPass, 
+					Context.HasHmdMesh(),
 					EDRF_UseTriangleOptimization);
 			}
 		}
diff --git a/Engine/Source/Runtime/Renderer/Private/DeferredShadingRenderer.cpp b/Engine/Source/Runtime/Renderer/Private/DeferredShadingRenderer.cpp
--- a/Engine/Source/Runtime/Renderer/Private/DeferredShadingRenderer.cpp
+++ b/Engine/Source/Runtime/Renderer/Private/DeferredShadingRenderer.cpp
@@ -692,10 +692,9 @@ static FORCEINLINE bool NeedsPrePass(const FDeferredShadingSceneRenderer* Render
  */
 static FORCEINLINE bool HasHiddenAreaMask()
 {
-	static const IConsoleVariable* const CVar = IConsoleManager::Get().FindConsoleVariable(TEXT("vr.HiddenAreaMask"));
-
-	return (CVar != nullptr &&
-		CVar->GetInt() == 1 &&
+	static const auto* const HiddenAreaMaskCVar = IConsoleManager::Get().FindTConsoleVariableDataInt(TEXT("vr.HiddenAreaMask"));
+	return (HiddenAreaMaskCVar != nullptr &&
+		HiddenAreaMaskCVar->GetValueOnRenderThread() == 1 &&
 		GEngine &&
 		GEngine->HMDDevice.IsValid() &&
 		GEngine->HMDDevice->HasHiddenAreaMask());
@@ -1430,7 +1429,7 @@ static void RenderHiddenAreaMaskView(FRHICommandList& RHICmdList, const FViewInf
 	TShaderMapRef<TOneColorVS<true> > VertexShader(ShaderMap);
 	static FGlobalBoundShaderState BoundShaderState;
 	SetGlobalBoundShaderState(RHICmdList, FeatureLevel, BoundShaderState, GetVertexDeclarationFVector4(), *VertexShader, nullptr);
-	GEngine->HMDDevice->DrawHiddenAreaMaskView_RenderThread(RHICmdList, View);
+	GEngine->HMDDevice->DrawHiddenAreaMaskView_RenderThread(RHICmdList, View.StereoPass);
 }
 
 bool FDeferredShadingSceneRenderer::RenderPrePassView(FRHICommandList& RHICmdList, const FViewInfo& View)
diff --git a/Engine/Source/Runtime/Renderer/Private/PostProcess/PostProcessAA.cpp b/Engine/Source/Runtime/Renderer/Private/PostProcess/PostProcessAA.cpp
--- a/Engine/Source/Runtime/Renderer/Private/PostProcess/PostProcessAA.cpp
+++ b/Engine/Source/Runtime/Renderer/Private/PostProcess/PostProcessAA.cpp
@@ -271,11 +271,9 @@ void FRCPassPostProcessAA::Process(FRenderingCompositePassContext& Context)
 		default: SetShaderTemplAA<6>(Context); break;
 	}
 	
-	// Draw a quad mapping scene color to the view's render target
 	TShaderMapRef<FFXAAVS> VertexShader(Context.GetShaderMap());
 
-	DrawRectangle(
-		Context.RHICmdList,
+	DrawPostProcessPass(Context.RHICmdList,
 		DestRect.Min.X, DestRect.Min.Y,
 		DestRect.Width(), DestRect.Height(),
 		SrcRect.Min.X, SrcRect.Min.Y,
@@ -283,6 +281,8 @@ void FRCPassPostProcessAA::Process(FRenderingCompositePassContext& Context)
 		DestSize,
 		SrcSize,
 		*VertexShader,
+		View.StereoPass,
+		Context.HasHmdMesh(),
 		EDRF_Default);
 
 	Context.RHICmdList.CopyToResolveTarget(DestRenderTarget.TargetableTexture, DestRenderTarget.ShaderResourceTexture, false, FResolveParams());
diff --git a/Engine/Source/Runtime/Renderer/Private/PostProcess/PostProcessBloomSetup.cpp b/Engine/Source/Runtime/Renderer/Private/PostProcess/PostProcessBloomSetup.cpp
--- a/Engine/Source/Runtime/Renderer/Private/PostProcess/PostProcessBloomSetup.cpp
+++ b/Engine/Source/Runtime/Renderer/Private/PostProcess/PostProcessBloomSetup.cpp
@@ -185,8 +185,7 @@ void FRCPassPostProcessBloomSetup::Process(FRenderingCompositePassContext& Conte
 	VertexShader->SetVS(Context);
 	PixelShader->SetPS(Context);
 
-	// Draw a quad mapping scene color to the view's render target
-	DrawRectangle(
+	DrawPostProcessPass(
 		Context.RHICmdList,
 		DestRect.Min.X, DestRect.Min.Y,
 		DestRect.Width(), DestRect.Height(),
@@ -195,6 +194,8 @@ void FRCPassPostProcessBloomSetup::Process(FRenderingCompositePassContext& Conte
 		DestSize,
 		SrcSize,
 		*VertexShader,
+		View.StereoPass,
+		Context.HasHmdMesh(),
 		EDRF_UseTriangleOptimization);
 
 	Context.RHICmdList.CopyToResolveTarget(DestRenderTarget.TargetableTexture, DestRenderTarget.ShaderResourceTexture, false, FResolveParams());
diff --git a/Engine/Source/Runtime/Renderer/Private/PostProcess/PostProcessBokehDOF.cpp b/Engine/Source/Runtime/Renderer/Private/PostProcess/PostProcessBokehDOF.cpp
--- a/Engine/Source/Runtime/Renderer/Private/PostProcess/PostProcessBokehDOF.cpp
+++ b/Engine/Source/Runtime/Renderer/Private/PostProcess/PostProcessBokehDOF.cpp
@@ -495,8 +495,7 @@ void FRCPassPostProcessBokehDOFSetup::Process(FRenderingCompositePassContext& Co
 		PixelShader->SetParameters(Context);
 	}
 
-	// Draw a quad mapping scene color to the view's render target
-	DrawRectangle(
+	DrawPostProcessPass(
 		Context.RHICmdList,
 		DestRect.Min.X, DestRect.Min.Y,
 		DestRect.Width(), DestRect.Height(),
@@ -505,6 +504,8 @@ void FRCPassPostProcessBokehDOFSetup::Process(FRenderingCompositePassContext& Co
 		DestSize,
 		SrcSize,
 		*VertexShader,
+		View.StereoPass,
+		Context.HasHmdMesh(),
 		EDRF_UseTriangleOptimization);
 
 	Context.RHICmdList.CopyToResolveTarget(DestRenderTarget.TargetableTexture, DestRenderTarget.ShaderResourceTexture, false, FResolveParams());
diff --git a/Engine/Source/Runtime/Renderer/Private/PostProcess/PostProcessBokehDOFRecombine.cpp b/Engine/Source/Runtime/Renderer/Private/PostProcess/PostProcessBokehDOFRecombine.cpp
--- a/Engine/Source/Runtime/Renderer/Private/PostProcess/PostProcessBokehDOFRecombine.cpp
+++ b/Engine/Source/Runtime/Renderer/Private/PostProcess/PostProcessBokehDOFRecombine.cpp
@@ -172,16 +172,17 @@ void FRCPassPostProcessBokehDOFRecombine::Process(FRenderingCompositePassContext
 
 	TShaderMapRef<FPostProcessVS> VertexShader(Context.GetShaderMap());
 
-	// Draw a quad mapping scene color to the view's render target
-	DrawRectangle(
+	DrawPostProcessPass(
 		Context.RHICmdList,
 		0, 0,
 		View.ViewRect.Width(), View.ViewRect.Height(),
-		HalfResViewRect.Min.X, HalfResViewRect.Min.Y, 
+		HalfResViewRect.Min.X, HalfResViewRect.Min.Y,
 		HalfResViewRect.Width(), HalfResViewRect.Height(),
 		View.ViewRect.Size(),
 		TexSize,
 		*VertexShader,
+		View.StereoPass,
+		Context.HasHmdMesh(),
 		EDRF_UseTriangleOptimization);
 
 	Context.RHICmdList.CopyToResolveTarget(DestRenderTarget.TargetableTexture, DestRenderTarget.ShaderResourceTexture, false, FResolveParams());
diff --git a/Engine/Source/Runtime/Renderer/Private/PostProcess/PostProcessCircleDOF.cpp b/Engine/Source/Runtime/Renderer/Private/PostProcess/PostProcessCircleDOF.cpp
--- a/Engine/Source/Runtime/Renderer/Private/PostProcess/PostProcessCircleDOF.cpp
+++ b/Engine/Source/Runtime/Renderer/Private/PostProcess/PostProcessCircleDOF.cpp
@@ -222,8 +222,7 @@ void FRCPassPostProcessCircleDOFSetup::Process(FRenderingCompositePassContext& C
 
 	VertexShader->SetParameters(Context);
 
-	// Draw a quad mapping scene color to the view's render target
-	DrawRectangle(
+	DrawPostProcessPass(
 		Context.RHICmdList,
 		DestRect.Min.X, DestRect.Min.Y,
 		DestRect.Width() + 1, DestRect.Height() + 1,
@@ -232,6 +231,8 @@ void FRCPassPostProcessCircleDOFSetup::Process(FRenderingCompositePassContext& C
 		DestSize,
 		SrcSize,
 		*VertexShader,
+		View.StereoPass,
+		Context.HasHmdMesh(),
 		EDRF_UseTriangleOptimization);
 
 	Context.RHICmdList.CopyToResolveTarget(DestRenderTarget0.TargetableTexture, DestRenderTarget0.ShaderResourceTexture, false, FResolveParams());
@@ -410,8 +411,7 @@ void FRCPassPostProcessCircleDOFDilate::Process(FRenderingCompositePassContext&
 
 	VertexShader->SetParameters(Context);
 
-	// Draw a quad mapping scene color to the view's render target
-	DrawRectangle(
+	DrawPostProcessPass(
 		Context.RHICmdList,
 		DestRect.Min.X, DestRect.Min.Y,
 		DestRect.Width() + 1, DestRect.Height() + 1,
@@ -420,8 +420,10 @@ void FRCPassPostProcessCircleDOFDilate::Process(FRenderingCompositePassContext&
 		DestSize,
 		SrcSize,
 		*VertexShader,
+		View.StereoPass,
+		Context.HasHmdMesh(),
 		EDRF_UseTriangleOptimization);
-
+	
 	Context.RHICmdList.CopyToResolveTarget(DestRenderTarget0.TargetableTexture, DestRenderTarget0.ShaderResourceTexture, false, FResolveParams());
 	Context.RHICmdList.CopyToResolveTarget(DestRenderTarget1.TargetableTexture, DestRenderTarget1.ShaderResourceTexture, false, FResolveParams());
 }
@@ -630,8 +632,7 @@ void FRCPassPostProcessCircleDOF::Process(FRenderingCompositePassContext& Contex
 
 	VertexShader->SetParameters(Context);
 
-	// Draw a quad mapping scene color to the view's render target
-	DrawRectangle(
+	DrawPostProcessPass(
 		Context.RHICmdList,
 		DestRect.Min.X, DestRect.Min.Y,
 		DestRect.Width() + 1, DestRect.Height() + 1,
@@ -640,6 +641,8 @@ void FRCPassPostProcessCircleDOF::Process(FRenderingCompositePassContext& Contex
 		DestSize,
 		SrcSize,
 		*VertexShader,
+		View.StereoPass,
+		Context.HasHmdMesh(),
 		EDRF_UseTriangleOptimization);
 
 	Context.RHICmdList.CopyToResolveTarget(DestRenderTarget0.TargetableTexture, DestRenderTarget0.ShaderResourceTexture, false, FResolveParams());
@@ -802,8 +805,7 @@ void FRCPassPostProcessCircleDOFRecombine::Process(FRenderingCompositePassContex
 
 	VertexShader->SetParameters(Context);
 
-	// Draw a quad mapping scene color to the view's render target
-	DrawRectangle(
+	DrawPostProcessPass(
 		Context.RHICmdList,
 		0, 0,
 		View.ViewRect.Width(), View.ViewRect.Height(),
@@ -812,6 +814,8 @@ void FRCPassPostProcessCircleDOFRecombine::Process(FRenderingCompositePassContex
 		View.ViewRect.Size(),
 		TexSize,
 		*VertexShader,
+		View.StereoPass,
+		Context.HasHmdMesh(),
 		EDRF_UseTriangleOptimization);
 
 	Context.RHICmdList.CopyToResolveTarget(DestRenderTarget.TargetableTexture, DestRenderTarget.ShaderResourceTexture, false, FResolveParams());
diff --git a/Engine/Source/Runtime/Renderer/Private/PostProcess/PostProcessDOF.cpp b/Engine/Source/Runtime/Renderer/Private/PostProcess/PostProcessDOF.cpp
--- a/Engine/Source/Runtime/Renderer/Private/PostProcess/PostProcessDOF.cpp
+++ b/Engine/Source/Runtime/Renderer/Private/PostProcess/PostProcessDOF.cpp
@@ -176,8 +176,7 @@ void FRCPassPostProcessDOFSetup::Process(FRenderingCompositePassContext& Context
 
 	VertexShader->SetParameters(Context);
 
-	// Draw a quad mapping scene color to the view's render target
-	DrawRectangle(
+	DrawPostProcessPass(
 		Context.RHICmdList,
 		0, 0,
 		DestRect.Width() + 1, DestRect.Height() + 1,
@@ -186,6 +185,8 @@ void FRCPassPostProcessDOFSetup::Process(FRenderingCompositePassContext& Context
 		DestRect.Size() + FIntPoint(1, 1),
 		SrcSize,
 		*VertexShader,
+		View.StereoPass,
+		Context.HasHmdMesh(),
 		EDRF_UseTriangleOptimization);
 
 	Context.RHICmdList.CopyToResolveTarget(DestRenderTarget0.TargetableTexture, DestRenderTarget0.ShaderResourceTexture, false, FResolveParams());
@@ -342,16 +343,17 @@ void FRCPassPostProcessDOFRecombine::Process(FRenderingCompositePassContext& Con
 
 	VertexShader->SetParameters(Context);
 
-	// Draw a quad mapping scene color to the view's render target
-	DrawRectangle(
+	DrawPostProcessPass(
 		Context.RHICmdList,
 		0, 0,
 		View.ViewRect.Width(), View.ViewRect.Height(),
-		HalfResViewRect.Min.X, HalfResViewRect.Min.Y, 
+		HalfResViewRect.Min.X, HalfResViewRect.Min.Y,
 		HalfResViewRect.Width(), HalfResViewRect.Height(),
 		View.ViewRect.Size(),
 		TexSize,
 		*VertexShader,
+		View.StereoPass,
+		Context.HasHmdMesh(),
 		EDRF_UseTriangleOptimization);
 
 	Context.RHICmdList.CopyToResolveTarget(DestRenderTarget.TargetableTexture, DestRenderTarget.ShaderResourceTexture, false, FResolveParams());
diff --git a/Engine/Source/Runtime/Renderer/Private/PostProcess/PostProcessDownsample.cpp b/Engine/Source/Runtime/Renderer/Private/PostProcess/PostProcessDownsample.cpp
--- a/Engine/Source/Runtime/Renderer/Private/PostProcess/PostProcessDownsample.cpp
+++ b/Engine/Source/Runtime/Renderer/Private/PostProcess/PostProcessDownsample.cpp
@@ -230,8 +230,7 @@ void FRCPassPostProcessDownsample::Process(FRenderingCompositePassContext& Conte
 		Context.RHICmdList.Clear(true, FLinearColor(0, 0, 0, 0), false, 1.0f, false, 0, DestRect);
 	}
 
-	// Draw a quad mapping scene color to the view's render target
-	DrawRectangle(
+	DrawPostProcessPass(
 		Context.RHICmdList,
 		DestRect.Min.X, DestRect.Min.Y,
 		DestRect.Width(), DestRect.Height(),
@@ -240,6 +239,8 @@ void FRCPassPostProcessDownsample::Process(FRenderingCompositePassContext& Conte
 		DestSize,
 		SrcSize,
 		*VertexShader,
+		View.StereoPass,
+		Context.HasHmdMesh(),
 		EDRF_UseTriangleOptimization);
 
 	Context.RHICmdList.CopyToResolveTarget(DestRenderTarget.TargetableTexture, DestRenderTarget.ShaderResourceTexture, false, FResolveParams());
diff --git a/Engine/Source/Runtime/Renderer/Private/PostProcess/PostProcessHierarchical.cpp b/Engine/Source/Runtime/Renderer/Private/PostProcess/PostProcessHierarchical.cpp
--- a/Engine/Source/Runtime/Renderer/Private/PostProcess/PostProcessHierarchical.cpp
+++ b/Engine/Source/Runtime/Renderer/Private/PostProcess/PostProcessHierarchical.cpp
@@ -164,8 +164,8 @@ void FRCPassPostProcessBuildHCB::Process(FRenderingCompositePassContext& Context
 			);
 
 		RHICmdList.SetViewport(0, 0, 0.0f, DstDrawSize.X, DstDrawSize.Y, 1.0f);
-			
-		DrawRectangle(
+		
+		DrawPostProcessPass(
 			RHICmdList,
 			0, 0,
 			HCBSize.X, HCBSize.Y,
@@ -174,11 +174,13 @@ void FRCPassPostProcessBuildHCB::Process(FRenderingCompositePassContext& Context
 			DstDrawSize,
 			SceneContext.GetBufferSizeXY(),
 			*VertexShader,
+			View.StereoPass,
+			Context.HasHmdMesh(),
 			EDRF_UseTriangleOptimization);
 #else
 		RHICmdList.SetViewport(0, 0, 0.0f, HCBSize.X, HCBSize.Y, 1.0f);
 			
-		DrawRectangle(
+		DrawPostProcessPass(
 			RHICmdList,
 			0, 0,
 			HCBSize.X, HCBSize.Y,
@@ -187,6 +189,8 @@ void FRCPassPostProcessBuildHCB::Process(FRenderingCompositePassContext& Context
 			HCBSize,
 			SceneContext.GetBufferSizeXY(),
 			*VertexShader,
+			View.StereoPass,
+			Context.HasHmdMesh(),
 			EDRF_UseTriangleOptimization);
 #endif
 
@@ -221,7 +225,7 @@ void FRCPassPostProcessBuildHCB::Process(FRenderingCompositePassContext& Context
 
 		RHICmdList.SetViewport(0, 0, 0.0f, DstDrawSize.X, DstDrawSize.Y, 1.0f);
 
-		DrawRectangle(
+		DrawPostProcessPass(
 			RHICmdList,
 			0, 0,
 			DstSize.X, DstSize.Y,
@@ -230,11 +234,13 @@ void FRCPassPostProcessBuildHCB::Process(FRenderingCompositePassContext& Context
 			DstDrawSize,
 			SrcSize,
 			*VertexShader,
+			View.StereoPass,
+			Context.HasHmdMesh(),
 			EDRF_UseTriangleOptimization);
 #else
 		RHICmdList.SetViewport(0, 0, 0.0f, DstSize.X, DstSize.Y, 1.0f);
 
-		DrawRectangle(
+		DrawPostProcessPass(
 			RHICmdList,
 			0, 0,
 			DstSize.X, DstSize.Y,
@@ -243,6 +249,8 @@ void FRCPassPostProcessBuildHCB::Process(FRenderingCompositePassContext& Context
 			DstSize,
 			SrcSize,
 			*VertexShader,
+			View.StereoPass,
+			Context.HasHmdMesh(),
 			EDRF_UseTriangleOptimization);
 #endif
 
diff --git a/Engine/Source/Runtime/Renderer/Private/PostProcess/PostProcessHistogramReduce.cpp b/Engine/Source/Runtime/Renderer/Private/PostProcess/PostProcessHistogramReduce.cpp
--- a/Engine/Source/Runtime/Renderer/Private/PostProcess/PostProcessHistogramReduce.cpp
+++ b/Engine/Source/Runtime/Renderer/Private/PostProcess/PostProcessHistogramReduce.cpp
@@ -128,8 +128,7 @@ void FRCPassPostProcessHistogramReduce::Process(FRenderingCompositePassContext&
 
 	PixelShader->SetPS(Context, LoopSizeValue);
 
-	// Draw a quad mapping scene color to the view's render target
-	DrawRectangle(
+	DrawPostProcessPass(
 		Context.RHICmdList,
 		0, 0,
 		DestSize.X, DestSize.Y,
@@ -138,6 +137,8 @@ void FRCPassPostProcessHistogramReduce::Process(FRenderingCompositePassContext&
 		DestSize,
 		SrcSize,
 		*VertexShader,
+		View.StereoPass,
+		Context.HasHmdMesh(),
 		EDRF_UseTriangleOptimization);
 
 	Context.RHICmdList.CopyToResolveTarget(DestRenderTarget.TargetableTexture, DestRenderTarget.ShaderResourceTexture, false, FResolveParams());
diff --git a/Engine/Source/Runtime/Renderer/Private/PostProcess/PostProcessMaterial.cpp b/Engine/Source/Runtime/Renderer/Private/PostProcess/PostProcessMaterial.cpp
--- a/Engine/Source/Runtime/Renderer/Private/PostProcess/PostProcessMaterial.cpp
+++ b/Engine/Source/Runtime/Renderer/Private/PostProcess/PostProcessMaterial.cpp
@@ -187,7 +187,7 @@ void FRCPassPostProcessMaterial::Process(FRenderingCompositePassContext& Context
 	VertexShader->SetParameters(Context.RHICmdList, Context);
 	PixelShader->SetParameters(Context.RHICmdList, Context, MaterialInterface->GetRenderProxy(false));
 
-	DrawRectangle(
+	DrawPostProcessPass(
 		Context.RHICmdList,
 		0, 0,
 		DestRect.Width(), DestRect.Height(),
@@ -196,6 +196,8 @@ void FRCPassPostProcessMaterial::Process(FRenderingCompositePassContext& Context
 		DestRect.Size(),
 		SrcSize,
 		VertexShader,
+		View.StereoPass,
+		Context.HasHmdMesh(),
 		EDRF_UseTriangleOptimization);
 
 	Context.RHICmdList.CopyToResolveTarget(DestRenderTarget.TargetableTexture, DestRenderTarget.ShaderResourceTexture, false, FResolveParams());
diff --git a/Engine/Source/Runtime/Renderer/Private/PostProcess/PostProcessMotionBlur.cpp b/Engine/Source/Runtime/Renderer/Private/PostProcess/PostProcessMotionBlur.cpp
--- a/Engine/Source/Runtime/Renderer/Private/PostProcess/PostProcessMotionBlur.cpp
+++ b/Engine/Source/Runtime/Renderer/Private/PostProcess/PostProcessMotionBlur.cpp
@@ -203,8 +203,7 @@ void FRCPassPostProcessMotionBlurSetup::Process(FRenderingCompositePassContext&
 		VertexShader->SetParameters(Context);
 	}
 
-	// Draw a quad mapping scene color to the view's render target
-	DrawRectangle(
+	DrawPostProcessPass(
 		Context.RHICmdList,
 		DestRect.Min.X, DestRect.Min.Y,
 		DestRect.Width(), DestRect.Height(),
@@ -213,6 +212,8 @@ void FRCPassPostProcessMotionBlurSetup::Process(FRenderingCompositePassContext&
 		DestSize,
 		SrcSize,
 		*VertexShader,
+		View.StereoPass,
+		Context.HasHmdMesh(),
 		EDRF_UseTriangleOptimization);
 
 	Context.RHICmdList.CopyToResolveTarget(DestRenderTarget0.TargetableTexture, DestRenderTarget0.ShaderResourceTexture, false, FResolveParams());
@@ -516,16 +517,17 @@ void FRCPassPostProcessMotionBlur::Process(FRenderingCompositePassContext& Conte
 
 	TShaderMapRef<FPostProcessVS> VertexShader(Context.GetShaderMap());
 
-	// Draw a quad mapping scene color to the view's render target
-	DrawRectangle(
+	DrawPostProcessPass(
 		Context.RHICmdList,
 		0, 0,
 		SrcRect.Width(), SrcRect.Height(),
-		SrcRect.Min.X, SrcRect.Min.Y, 
+		SrcRect.Min.X, SrcRect.Min.Y,
 		SrcRect.Width(), SrcRect.Height(),
 		SrcRect.Size(),
 		SrcSize,
 		*VertexShader,
+		View.StereoPass,
+		Context.HasHmdMesh(),
 		EDRF_UseTriangleOptimization);
 
 	Context.RHICmdList.CopyToResolveTarget(DestRenderTarget.TargetableTexture, DestRenderTarget.ShaderResourceTexture, false, FResolveParams());
@@ -651,16 +653,17 @@ void FRCPassPostProcessMotionBlurRecombine::Process(FRenderingCompositePassConte
 	VertexShader->SetParameters(Context);
 	PixelShader->SetParameters(Context, bBilinear);
 
-	// Draw a quad mapping scene color to the view's render target
-	DrawRectangle(
+	DrawPostProcessPass(
 		Context.RHICmdList,
 		0, 0,
 		SrcRect.Width(), SrcRect.Height(),
-		SrcRect.Min.X, SrcRect.Min.Y, 
+		SrcRect.Min.X, SrcRect.Min.Y,
 		SrcRect.Width(), SrcRect.Height(),
 		SrcRect.Size(),
 		SrcSize,
 		*VertexShader,
+		View.StereoPass,
+		Context.HasHmdMesh(),
 		EDRF_UseTriangleOptimization);
 
 	Context.RHICmdList.CopyToResolveTarget(DestRenderTarget.TargetableTexture, DestRenderTarget.ShaderResourceTexture, false, FResolveParams());
@@ -1402,16 +1405,18 @@ void FRCPassPostProcessMotionBlurNew::Process(FRenderingCompositePassContext& Co
 
 	TShaderMapRef<FPostProcessVS> VertexShader(Context.GetShaderMap());
 
-	// Draw a quad mapping scene color to the view's render target
-	DrawRectangle(
+	
+	DrawPostProcessPass(
 		Context.RHICmdList,
 		0, 0,
 		SrcRect.Width(), SrcRect.Height(),
-		SrcRect.Min.X, SrcRect.Min.Y, 
+		SrcRect.Min.X, SrcRect.Min.Y,
 		SrcRect.Width(), SrcRect.Height(),
 		SrcRect.Size(),
 		SrcSize,
 		*VertexShader,
+		View.StereoPass,
+		Context.HasHmdMesh(),
 		EDRF_UseTriangleOptimization);
 
 	Context.RHICmdList.CopyToResolveTarget(DestRenderTarget.TargetableTexture, DestRenderTarget.ShaderResourceTexture, false, FResolveParams());
diff --git a/Engine/Source/Runtime/Renderer/Private/PostProcess/PostProcessNoiseBlur.cpp b/Engine/Source/Runtime/Renderer/Private/PostProcess/PostProcessNoiseBlur.cpp
--- a/Engine/Source/Runtime/Renderer/Private/PostProcess/PostProcessNoiseBlur.cpp
+++ b/Engine/Source/Runtime/Renderer/Private/PostProcess/PostProcessNoiseBlur.cpp
@@ -162,9 +162,9 @@ void FRCPassPostProcessNoiseBlur::Process(FRenderingCompositePassContext& Contex
 		SetNoiseBlurShader<2>(Context, Radius);
 	}
 
-	// Draw a quad mapping scene color to the view's render target
 	TShaderMapRef<FPostProcessVS> VertexShader(Context.GetShaderMap());
-	DrawRectangle(
+
+	DrawPostProcessPass(
 		Context.RHICmdList,
 		DestRect.Min.X, DestRect.Min.Y,
 		DestRect.Width(), DestRect.Height(),
@@ -173,6 +173,8 @@ void FRCPassPostProcessNoiseBlur::Process(FRenderingCompositePassContext& Contex
 		DestSize,
 		SrcSize,
 		*VertexShader,
+		View.StereoPass,
+		Context.HasHmdMesh(),
 		EDRF_UseTriangleOptimization);
 
 	Context.RHICmdList.CopyToResolveTarget(DestRenderTarget.TargetableTexture, DestRenderTarget.ShaderResourceTexture, false, FResolveParams());
diff --git a/Engine/Source/Runtime/Renderer/Private/PostProcess/PostProcessSceneColorFringe.cpp b/Engine/Source/Runtime/Renderer/Private/PostProcess/PostProcessSceneColorFringe.cpp
--- a/Engine/Source/Runtime/Renderer/Private/PostProcess/PostProcessSceneColorFringe.cpp
+++ b/Engine/Source/Runtime/Renderer/Private/PostProcess/PostProcessSceneColorFringe.cpp
@@ -178,16 +178,17 @@ void FRCPassPostProcessSceneColorFringe::Process(FRenderingCompositePassContext&
 	PixelShader->SetParameters(Context);
 	VertexShader->SetParameters(Context);
 	
-	// Draw a quad mapping scene color to the view's render target
-	DrawRectangle(
+	DrawPostProcessPass(
 		Context.RHICmdList,
 		0, 0,
 		View.ViewRect.Width(), View.ViewRect.Height(),
-		View.ViewRect.Min.X, View.ViewRect.Min.Y, 
+		View.ViewRect.Min.X, View.ViewRect.Min.Y,
 		View.ViewRect.Width(), View.ViewRect.Height(),
 		View.ViewRect.Size(),
 		FSceneRenderTargets::Get(Context.RHICmdList).GetBufferSizeXY(),
 		*VertexShader,
+		View.StereoPass,
+		Context.HasHmdMesh(),
 		EDRF_UseTriangleOptimization);
 
 	Context.RHICmdList.CopyToResolveTarget(DestRenderTarget.TargetableTexture, DestRenderTarget.ShaderResourceTexture, false, FResolveParams());
diff --git a/Engine/Source/Runtime/Renderer/Private/PostProcess/PostProcessSubsurface.cpp b/Engine/Source/Runtime/Renderer/Private/PostProcess/PostProcessSubsurface.cpp
--- a/Engine/Source/Runtime/Renderer/Private/PostProcess/PostProcessSubsurface.cpp
+++ b/Engine/Source/Runtime/Renderer/Private/PostProcess/PostProcessSubsurface.cpp
@@ -463,7 +463,8 @@ void FRCPassPostProcessSubsurfaceSetup::Process(FRenderingCompositePassContext&
 
 	// Draw a quad mapping scene color to the view's render target
 	TShaderMapRef<FPostProcessVS> VertexShader(Context.GetShaderMap());
-	DrawRectangle(
+
+	DrawPostProcessPass(
 		Context.RHICmdList,
 		DestRect.Min.X, DestRect.Min.Y,
 		DestRect.Width(), DestRect.Height(),
@@ -472,9 +473,10 @@ void FRCPassPostProcessSubsurfaceSetup::Process(FRenderingCompositePassContext&
 		DestSize,
 		SrcSize,
 		*VertexShader,
+		View.StereoPass,
+		Context.HasHmdMesh(),
 		EDRF_UseTriangleOptimization);
 
-
 	Context.RHICmdList.CopyToResolveTarget(DestRenderTarget.TargetableTexture, DestRenderTarget.ShaderResourceTexture, false, FResolveParams());
 }
 
@@ -682,7 +684,7 @@ void FRCPassPostProcessSubsurface::Process(FRenderingCompositePassContext& Conte
 		SetSubsurfaceShaderSampleSet<1>(Context, VertexShader, SampleSet);
 	}
 
-	DrawRectangle(
+	DrawPostProcessPass(
 		Context.RHICmdList,
 		DestRect.Min.X, DestRect.Min.Y,
 		DestRect.Width(), DestRect.Height(),
@@ -691,6 +693,8 @@ void FRCPassPostProcessSubsurface::Process(FRenderingCompositePassContext& Conte
 		DestSize,
 		SrcSize,
 		*VertexShader,
+		View.StereoPass,
+		Context.HasHmdMesh(),
 		EDRF_UseTriangleOptimization);
 
 	Context.RHICmdList.CopyToResolveTarget(DestRenderTarget->TargetableTexture, DestRenderTarget->ShaderResourceTexture, false, FResolveParams());
@@ -862,7 +866,7 @@ void FRCPassPostProcessSubsurfaceRecombine::Process(FRenderingCompositePassConte
 		SetSubsurfaceRecombineShader<2>(Context, VertexShader);
 	}
 
-	DrawRectangle(
+	DrawPostProcessPass(
 		Context.RHICmdList,
 		DestRect.Min.X, DestRect.Min.Y,
 		DestRect.Width(), DestRect.Height(),
@@ -871,6 +875,8 @@ void FRCPassPostProcessSubsurfaceRecombine::Process(FRenderingCompositePassConte
 		DestSize,
 		SrcSize,
 		*VertexShader,
+		View.StereoPass,
+		Context.HasHmdMesh(),
 		EDRF_UseTriangleOptimization);
 
 	Context.RHICmdList.CopyToResolveTarget(DestRenderTarget.TargetableTexture, DestRenderTarget.ShaderResourceTexture, false, FResolveParams());
diff --git a/Engine/Source/Runtime/Renderer/Private/PostProcess/PostProcessTemporalAA.cpp b/Engine/Source/Runtime/Renderer/Private/PostProcess/PostProcessTemporalAA.cpp
--- a/Engine/Source/Runtime/Renderer/Private/PostProcess/PostProcessTemporalAA.cpp
+++ b/Engine/Source/Runtime/Renderer/Private/PostProcess/PostProcessTemporalAA.cpp
@@ -268,16 +268,17 @@ void FRCPassPostProcessSSRTemporalAA::Process(FRenderingCompositePassContext& Co
 	VertexShader->SetVS(Context);
 	PixelShader->SetParameters(Context);
 
-	// Draw a quad mapping scene color to the view's render target
-	DrawRectangle(
+	DrawPostProcessPass(
 		Context.RHICmdList,
 		0, 0,
 		SrcRect.Width(), SrcRect.Height(),
-		SrcRect.Min.X, SrcRect.Min.Y, 
+		SrcRect.Min.X, SrcRect.Min.Y,
 		SrcRect.Width(), SrcRect.Height(),
 		SrcRect.Size(),
 		SrcSize,
 		*VertexShader,
+		View.StereoPass,
+		Context.HasHmdMesh(),
 		EDRF_UseTriangleOptimization);
 
 	Context.RHICmdList.CopyToResolveTarget(DestRenderTarget.TargetableTexture, DestRenderTarget.ShaderResourceTexture, false, FResolveParams());
@@ -345,16 +346,17 @@ void FRCPassPostProcessDOFTemporalAA::Process(FRenderingCompositePassContext& Co
 	VertexShader->SetVS(Context);
 	PixelShader->SetParameters(Context);
 
-	// Draw a quad mapping scene color to the view's render target
-	DrawRectangle(
+	DrawPostProcessPass(
 		Context.RHICmdList,
 		0, 0,
 		SrcRect.Width(), SrcRect.Height(),
-		SrcRect.Min.X, SrcRect.Min.Y, 
+		SrcRect.Min.X, SrcRect.Min.Y,
 		SrcRect.Width(), SrcRect.Height(),
 		SrcRect.Size(),
 		SrcSize,
 		*VertexShader,
+		View.StereoPass,
+		Context.HasHmdMesh(),
 		EDRF_UseTriangleOptimization);
 
 	Context.RHICmdList.CopyToResolveTarget(DestRenderTarget.TargetableTexture, DestRenderTarget.ShaderResourceTexture, false, FResolveParams());
@@ -426,16 +428,17 @@ void FRCPassPostProcessDOFTemporalAANear::Process(FRenderingCompositePassContext
 	VertexShader->SetVS(Context);
 	PixelShader->SetParameters(Context);
 
-	// Draw a quad mapping scene color to the view's render target
-	DrawRectangle(
+	DrawPostProcessPass(
 		Context.RHICmdList,
 		0, 0,
 		SrcRect.Width(), SrcRect.Height(),
-		SrcRect.Min.X, SrcRect.Min.Y, 
+		SrcRect.Min.X, SrcRect.Min.Y,
 		SrcRect.Width(), SrcRect.Height(),
 		SrcRect.Size(),
 		SrcSize,
 		*VertexShader,
+		View.StereoPass,
+		Context.HasHmdMesh(),
 		EDRF_UseTriangleOptimization);
 
 	Context.RHICmdList.CopyToResolveTarget(DestRenderTarget.TargetableTexture, DestRenderTarget.ShaderResourceTexture, false, FResolveParams());
@@ -506,16 +509,17 @@ void FRCPassPostProcessLightShaftTemporalAA::Process(FRenderingCompositePassCont
 	VertexShader->SetVS(Context);
 	PixelShader->SetParameters(Context);
 
-	// Draw a quad mapping scene color to the view's render target
-	DrawRectangle(
+	DrawPostProcessPass(
 		Context.RHICmdList,
 		0, 0,
 		SrcRect.Width(), SrcRect.Height(),
-		SrcRect.Min.X, SrcRect.Min.Y, 
+		SrcRect.Min.X, SrcRect.Min.Y,
 		SrcRect.Width(), SrcRect.Height(),
 		SrcRect.Size(),
 		SrcSize,
 		*VertexShader,
+		View.StereoPass,
+		Context.HasHmdMesh(),
 		EDRF_UseTriangleOptimization);
 
 	Context.RHICmdList.CopyToResolveTarget(DestRenderTarget.TargetableTexture, DestRenderTarget.ShaderResourceTexture, false, FResolveParams());
@@ -611,16 +615,17 @@ void FRCPassPostProcessTemporalAA::Process(FRenderingCompositePassContext& Conte
 			PixelShader->SetParameters(Context);
 		}
 	
-		// Draw a quad mapping scene color to the view's render target
-		DrawRectangle(
+		DrawPostProcessPass(
 			Context.RHICmdList,
 			0, 0,
 			SrcRect.Width(), SrcRect.Height(),
-			SrcRect.Min.X, SrcRect.Min.Y, 
+			SrcRect.Min.X, SrcRect.Min.Y,
 			SrcRect.Width(), SrcRect.Height(),
 			SrcRect.Size(),
 			SrcSize,
 			*VertexShader,
+			View.StereoPass,
+			Context.HasHmdMesh(),
 			EDRF_UseTriangleOptimization);
 	}
 	else
@@ -656,16 +661,17 @@ void FRCPassPostProcessTemporalAA::Process(FRenderingCompositePassContext& Conte
 				PixelShader->SetParameters(Context);
 			}
 		
-			// Draw a quad mapping scene color to the view's render target
-			DrawRectangle(
+			DrawPostProcessPass(
 				Context.RHICmdList,
 				0, 0,
 				SrcRect.Width(), SrcRect.Height(),
-				SrcRect.Min.X, SrcRect.Min.Y, 
+				SrcRect.Min.X, SrcRect.Min.Y,
 				SrcRect.Width(), SrcRect.Height(),
 				SrcRect.Size(),
 				SrcSize,
 				*VertexShader,
+				View.StereoPass,
+				Context.HasHmdMesh(),
 				EDRF_UseTriangleOptimization);
 		}
 	
@@ -699,16 +705,17 @@ void FRCPassPostProcessTemporalAA::Process(FRenderingCompositePassContext& Conte
 				PixelShader->SetParameters(Context);
 			}
 	
-			// Draw a quad mapping scene color to the view's render target
-			DrawRectangle(
+			DrawPostProcessPass(
 				Context.RHICmdList,
 				0, 0,
 				SrcRect.Width(), SrcRect.Height(),
-				SrcRect.Min.X, SrcRect.Min.Y, 
+				SrcRect.Min.X, SrcRect.Min.Y,
 				SrcRect.Width(), SrcRect.Height(),
 				SrcRect.Size(),
 				SrcSize,
 				*VertexShader,
+				View.StereoPass,
+				Context.HasHmdMesh(),
 				EDRF_UseTriangleOptimization);
 		}
 	}
diff --git a/Engine/Source/Runtime/Renderer/Private/PostProcess/PostProcessTonemap.cpp b/Engine/Source/Runtime/Renderer/Private/PostProcess/PostProcessTonemap.cpp
--- a/Engine/Source/Runtime/Renderer/Private/PostProcess/PostProcessTonemap.cpp
+++ b/Engine/Source/Runtime/Renderer/Private/PostProcess/PostProcessTonemap.cpp
@@ -1126,18 +1126,19 @@ void FRCPassPostProcessTonemap::Process(FRenderingCompositePassContext& Context)
 	}
 	FSceneRenderTargets& SceneContext = FSceneRenderTargets::Get(Context.RHICmdList);
 
-	// Draw a quad mapping scene color to the view's render target
 	TShaderMapRef<FPostProcessTonemapVS> VertexShader(Context.GetShaderMap());
 
-	DrawRectangle(
+	DrawPostProcessPass(
 		Context.RHICmdList,
 		0, 0,
 		View.ViewRect.Width(), View.ViewRect.Height(),
-		View.ViewRect.Min.X, View.ViewRect.Min.Y, 
+		View.ViewRect.Min.X, View.ViewRect.Min.Y,
 		View.ViewRect.Width(), View.ViewRect.Height(),
 		View.ViewRect.Size(),
 		SceneContext.GetBufferSizeXY(),
 		*VertexShader,
+		View.StereoPass,
+		Context.HasHmdMesh(),
 		EDRF_UseTriangleOptimization);
 
 	Context.RHICmdList.CopyToResolveTarget(DestRenderTarget.TargetableTexture, DestRenderTarget.ShaderResourceTexture, false, FResolveParams());
diff --git a/Engine/Source/Runtime/Renderer/Private/PostProcess/RenderingCompositionGraph.cpp b/Engine/Source/Runtime/Renderer/Private/PostProcess/RenderingCompositionGraph.cpp
--- a/Engine/Source/Runtime/Renderer/Private/PostProcess/RenderingCompositionGraph.cpp
+++ b/Engine/Source/Runtime/Renderer/Private/PostProcess/RenderingCompositionGraph.cpp
@@ -7,6 +7,7 @@
 #include "RendererPrivate.h"
 #include "RenderingCompositionGraph.h"
 #include "HighResScreenshot.h"
+#include "IHeadMountedDisplay.h"
 
 void ExecuteCompositionGraphDebug();
 
@@ -123,6 +124,7 @@ FRenderingCompositePassContext::FRenderingCompositePassContext(FRHICommandListIm
 	, FeatureLevel(View.GetFeatureLevel())
 	, ShaderMap(InView.ShaderMap)
 	, bWasProcessed(false)
+	, bHasHmdMesh(false)
 {
 	check(!IsViewportValid());
 }
@@ -134,11 +136,19 @@ FRenderingCompositePassContext::~FRenderingCompositePassContext()
 
 void FRenderingCompositePassContext::Process(FRenderingCompositePass* Root, const TCHAR *GraphDebugName)
 {
-	// call this method only once afetr the graph is finished
+	// call this method only once after the graph is finished
 	check(!bWasProcessed);
 
 	bWasProcessed = true;
 
+	// query if we have a custom HMD post process mesh to use
+	static const auto* const HiddenAreaMaskCVar = IConsoleManager::Get().FindTConsoleVariableDataInt(TEXT("vr.HiddenAreaMask"));
+	bHasHmdMesh = (HiddenAreaMaskCVar != nullptr &&
+		HiddenAreaMaskCVar->GetValueOnRenderThread() == 1 &&
+		GEngine &&
+		GEngine->HMDDevice.IsValid() &&
+		GEngine->HMDDevice->HasPostProcessMesh());
+
 	if(Root)
 	{
 		if(ShouldDebugCompositionGraph())
diff --git a/Engine/Source/Runtime/Renderer/Private/PostProcess/RenderingCompositionGraph.h b/Engine/Source/Runtime/Renderer/Private/PostProcess/RenderingCompositionGraph.h
--- a/Engine/Source/Runtime/Renderer/Private/PostProcess/RenderingCompositionGraph.h
+++ b/Engine/Source/Runtime/Renderer/Private/PostProcess/RenderingCompositionGraph.h
@@ -129,6 +129,11 @@ struct FRenderingCompositePassContext
 		return ViewPortRect.Min != ViewPortRect.Max;
 	}
 
+	bool HasHmdMesh() const
+	{
+		return bHasHmdMesh;
+	}
+
 	ERHIFeatureLevel::Type GetFeatureLevel() const { return FeatureLevel; }
 	EShaderPlatform GetShaderPlatform() const { return GShaderPlatformForFeatureLevel[FeatureLevel]; }
 	TShaderMap<FGlobalShaderType>* GetShaderMap() const { check(ShaderMap); return ShaderMap; }
@@ -154,6 +159,9 @@ private:
 	TShaderMap<FGlobalShaderType>* ShaderMap;
 	// to ensure we only process the graph once
 	bool bWasProcessed;
+	// updated once a frame in Process()
+	// If true there's a custom mesh to use instead of a full screen quad when rendering post process passes.
+	bool bHasHmdMesh;
 };
 
 // ---------------------------------------------------------------------------
diff --git a/Engine/Source/Runtime/Renderer/Private/PostProcess/SceneFilterRendering.cpp b/Engine/Source/Runtime/Renderer/Private/PostProcess/SceneFilterRendering.cpp
--- a/Engine/Source/Runtime/Renderer/Private/PostProcess/SceneFilterRendering.cpp
+++ b/Engine/Source/Runtime/Renderer/Private/PostProcess/SceneFilterRendering.cpp
@@ -8,6 +8,7 @@
 #include "ScenePrivate.h"
 #include "SceneFilterRendering.h"
 #include "PostProcessWeightedSampleSum.h"
+#include "IHeadMountedDisplay.h"
 
 /**
 * Static vertex and index buffer used for 2D screen rectangles.
@@ -300,3 +301,59 @@ void DrawTransformedRectangle(
 
 	DrawIndexedPrimitiveUP(RHICmdList, PT_TriangleList, 0, 4, 2, Indices, sizeof(Indices[0]), Vertices, sizeof(Vertices[0]));
 }
+
+void DrawHmdMesh(
+	FRHICommandList& RHICmdList,
+	float X,
+	float Y,
+	float SizeX,
+	float SizeY,
+	float U,
+	float V,
+	float SizeU,
+	float SizeV,
+	FIntPoint TargetSize,
+	FIntPoint TextureSize,
+	EStereoscopicPass StereoView,
+	FShader* VertexShader
+	)
+{
+	FDrawRectangleParameters Parameters;
+	Parameters.PosScaleBias = FVector4(SizeX, SizeY, X, Y);
+	Parameters.UVScaleBias = FVector4(SizeU, SizeV, U, V);
+
+	Parameters.InvTargetSizeAndTextureSize = FVector4(
+		1.0f / TargetSize.X, 1.0f / TargetSize.Y,
+		1.0f / TextureSize.X, 1.0f / TextureSize.Y);
+
+	SetUniformBufferParameterImmediate(RHICmdList, VertexShader->GetVertexShader(), VertexShader->GetUniformBufferParameter<FDrawRectangleParameters>(), Parameters);
+
+	GEngine->HMDDevice->DrawPostProcessMesh_RenderThread(RHICmdList, StereoView);
+}
+
+void DrawPostProcessPass(
+	FRHICommandList& RHICmdList,
+	float X,
+	float Y,
+	float SizeX,
+	float SizeY,
+	float U,
+	float V,
+	float SizeU,
+	float SizeV,
+	FIntPoint TargetSize,
+	FIntPoint TextureSize,
+	FShader* VertexShader,
+	EStereoscopicPass StereoView,
+	bool bHasCustomMesh,
+	EDrawRectangleFlags Flags)
+{
+	if (bHasCustomMesh && StereoView != eSSP_FULL)
+	{
+		DrawHmdMesh(RHICmdList, X, Y, SizeX, SizeY, U, V, SizeU, SizeV, TargetSize, TextureSize, StereoView, VertexShader);
+	}
+	else
+	{
+		DrawRectangle(RHICmdList, X, Y, SizeX, SizeY, U, V, SizeU, SizeV, TargetSize, TextureSize, VertexShader, Flags);
+	}
+}
diff --git a/Engine/Source/Runtime/Renderer/Private/PostProcess/SceneFilterRendering.h b/Engine/Source/Runtime/Renderer/Private/PostProcess/SceneFilterRendering.h
--- a/Engine/Source/Runtime/Renderer/Private/PostProcess/SceneFilterRendering.h
+++ b/Engine/Source/Runtime/Renderer/Private/PostProcess/SceneFilterRendering.h
@@ -56,6 +56,40 @@ extern void DrawTransformedRectangle(
     FIntPoint TextureSize
     );
 
+extern void DrawHmdMesh(
+	FRHICommandList& RHICmdList,
+	float X,
+	float Y,
+	float SizeX,
+	float SizeY,
+	float U,
+	float V,
+	float SizeU,
+	float SizeV,
+	FIntPoint TargetSize,
+	FIntPoint TextureSize,
+	EStereoscopicPass StereoView,
+	FShader* VertexShader
+	);
+
+extern void DrawPostProcessPass(
+	FRHICommandList& RHICmdList,
+	float X,
+	float Y,
+	float SizeX,
+	float SizeY,
+	float U,
+	float V,
+	float SizeU,
+	float SizeV,
+	FIntPoint TargetSize,
+	FIntPoint TextureSize,
+	class FShader* VertexShader,
+	EStereoscopicPass StereoView,
+	bool bHasCustomMesh,
+	EDrawRectangleFlags Flags = EDRF_Default
+	);
+
 extern TGlobalResource<FFilterVertexDeclaration> GFilterVertexDeclaration;
 extern TGlobalResource<FEmptyVertexDeclaration> GEmptyVertexDeclaration;
 
diff --git a/Engine/Source/Runtime/Renderer/Private/PostProcess/ScreenSpaceReflections.cpp b/Engine/Source/Runtime/Renderer/Private/PostProcess/ScreenSpaceReflections.cpp
--- a/Engine/Source/Runtime/Renderer/Private/PostProcess/ScreenSpaceReflections.cpp
+++ b/Engine/Source/Runtime/Renderer/Private/PostProcess/ScreenSpaceReflections.cpp
@@ -348,17 +348,19 @@ void FRCPassPostProcessScreenSpaceReflections::Process(FRenderingCompositePassCo
 		// disable blend mode
 		RHICmdList.SetBlendState(TStaticBlendState<>::GetRHI());
 	
-		// Draw a quad mapping scene color to the view's render target to set stencil to set the stencil mask where it needs to be
-		DrawRectangle( 
+		DrawPostProcessPass(
 			Context.RHICmdList,
 			0, 0,
 			View.ViewRect.Width(), View.ViewRect.Height(),
-			View.ViewRect.Min.X, View.ViewRect.Min.Y, 
+			View.ViewRect.Min.X, View.ViewRect.Min.Y,
 			View.ViewRect.Width(), View.ViewRect.Height(),
 			View.ViewRect.Size(),
 			SceneContext.GetBufferSizeXY(),
 			*VertexShader,
+			View.StereoPass,
+			Context.HasHmdMesh(),
 			EDRF_UseTriangleOptimization);
+
 	} // ScreenSpaceReflectionsStencil draw event
 
 	{ // ScreenSpaceReflections draw event
@@ -411,19 +413,19 @@ void FRCPassPostProcessScreenSpaceReflections::Process(FRenderingCompositePassCo
 		}
 		#undef CASE
 
-
-		// Draw a quad mapping scene color to the view's render target
-		DrawRectangle( 
+		DrawPostProcessPass(
 			RHICmdList,
 			0, 0,
 			View.ViewRect.Width(), View.ViewRect.Height(),
-			View.ViewRect.Min.X, View.ViewRect.Min.Y, 
+			View.ViewRect.Min.X, View.ViewRect.Min.Y,
 			View.ViewRect.Width(), View.ViewRect.Height(),
 			View.ViewRect.Size(),
 			FSceneRenderTargets::Get(Context.RHICmdList).GetBufferSizeXY(),
 			*VertexShader,
+			View.StereoPass,
+			Context.HasHmdMesh(),
 			EDRF_UseTriangleOptimization);
-
+	
 		RHICmdList.CopyToResolveTarget(DestRenderTarget.TargetableTexture, DestRenderTarget.ShaderResourceTexture, false, FResolveParams());
 	} // ScreenSpaceReflections
 }
