[INTEGRATE] Compile times: Fixed HeadMountedDisplay module not using shared PCHs

diff --git a/Engine/Source/Runtime/HeadMountedDisplay/HeadMountedDisplay.Build.cs b/Engine/Source/Runtime/HeadMountedDisplay/HeadMountedDisplay.Build.cs
--- a/Engine/Source/Runtime/HeadMountedDisplay/HeadMountedDisplay.Build.cs
+++ b/Engine/Source/Runtime/HeadMountedDisplay/HeadMountedDisplay.Build.cs
@@ -21,7 +21,5 @@ public class HeadMountedDisplay : ModuleRules
 				"SlateCore",
 			}
 		);
-
-        PCHUsage = PCHUsageMode.NoSharedPCHs;
 	}
 }
