#ue4 Motion controller late update.

-Applies a more up-to-date controller transform to dependent visible objects immediately before rendering occurs
-Can be disabled on a per motion controller component basis through a blueprint option "Disable Low Latency Update"
-Can be toggled holistically by a console command "vr.EnableMotionControllerLateUpdate"

diff --git a/Engine/Source/Runtime/Engine/Private/PrimitiveSceneProxy.cpp b/Engine/Source/Runtime/Engine/Private/PrimitiveSceneProxy.cpp
--- a/Engine/Source/Runtime/Engine/Private/PrimitiveSceneProxy.cpp
+++ b/Engine/Source/Runtime/Engine/Private/PrimitiveSceneProxy.cpp
@@ -264,6 +264,12 @@ void FPrimitiveSceneProxy::ApplyWorldOffset(FVector InOffset)
 	SetTransform(NewLocalToWorld, NewBounds, NewLocalBounds, NewActorPosition);
 }
 
+void FPrimitiveSceneProxy::ApplyLateUpdateTransform(const FMatrix& LateUpdateTransform)
+{
+	const FMatrix AdjustedLocalToWorld = LocalToWorld * LateUpdateTransform;
+	SetTransform(AdjustedLocalToWorld, Bounds, LocalBounds, ActorPosition);
+}
+
 /**
  * Updates selection for the primitive proxy. This is called in the rendering thread by SetSelection_GameThread.
  * @param bInSelected - true if the parent actor is selected in the editor
diff --git a/Engine/Source/Runtime/Engine/Public/PrimitiveSceneProxy.h b/Engine/Source/Runtime/Engine/Public/PrimitiveSceneProxy.h
--- a/Engine/Source/Runtime/Engine/Public/PrimitiveSceneProxy.h
+++ b/Engine/Source/Runtime/Engine/Public/PrimitiveSceneProxy.h
@@ -484,6 +484,12 @@ public:
 	 */
 	ENGINE_API virtual void ApplyWorldOffset(FVector InOffset);
 
+	/**
+	 * Applies a "late in the frame" adjustment to the proxy's existing transform
+	 * @param LateUpdateTransform - The post-transform to be applied to the LocalToWorld matrix
+	 */
+	ENGINE_API virtual void ApplyLateUpdateTransform(const FMatrix& LateUpdateTransform);
+
 	/**
 	 * Updates the primitive proxy's uniform buffer.
 	 */
diff --git a/Engine/Source/Runtime/HeadMountedDisplay/Private/MotionControllerComponent.cpp b/Engine/Source/Runtime/HeadMountedDisplay/Private/MotionControllerComponent.cpp
--- a/Engine/Source/Runtime/HeadMountedDisplay/Private/MotionControllerComponent.cpp
+++ b/Engine/Source/Runtime/HeadMountedDisplay/Private/MotionControllerComponent.cpp
@@ -3,6 +3,22 @@
 #include "HeadMountedDisplayPrivate.h"
 #include "MotionControllerComponent.h"
 
+namespace {
+	/** This is to prevent destruction of motion controller components while they are
+		in the middle of being accessed by the render thread */
+	FCriticalSection CritSect;
+
+	/** Console variable for specifying whether motion controller late update is used */
+	TAutoConsoleVariable<int32> CVarEnableMotionControllerLateUpdate(
+		TEXT("vr.EnableMotionControllerLateUpdate"),
+		1,
+		TEXT("This command allows you to specify whether the motion controller late update is applied.\n")
+		TEXT(" 0: don't use late update\n")
+		TEXT(" 1: use late update (default)"),
+		ECVF_Cheat);
+} // anonymous namespace
+
+//=============================================================================
 UMotionControllerComponent::UMotionControllerComponent(const FObjectInitializer& ObjectInitializer)
 : Super(ObjectInitializer)
 {
@@ -12,16 +28,53 @@ UMotionControllerComponent::UMotionControllerComponent(const FObjectInitializer&
 
 	PlayerIndex = 0;
 	Hand = EControllerHand::Left;
+	bDisableLowLatencyUpdate = false;
 }
 
+//=============================================================================
+UMotionControllerComponent::~UMotionControllerComponent()
+{
+	if (ViewExtension.IsValid())
+	{
+		{
+			// This component could be getting accessed from the render thread so it needs to wait
+			// before clearing MotionControllerComponent and allowing the destructor to continue
+			FScopeLock ScopeLock(&CritSect);
+			ViewExtension->MotionControllerComponent = NULL;
+		}
+
+		if (GEngine)
+		{
+			GEngine->ViewExtensions.Remove(ViewExtension);
+		}
+	}
+	ViewExtension.Reset();
+}
+
+//=============================================================================
 void UMotionControllerComponent::TickComponent(float DeltaTime, enum ELevelTick TickType, FActorComponentTickFunction *ThisTickFunction)
 {
 	Super::TickComponent(DeltaTime, TickType, ThisTickFunction);
 
-	bTracked = false;
-	FVector Position = FVector::ZeroVector;
-	FRotator Orientation = FRotator::ZeroRotator;
+	FVector Position;
+	FRotator Orientation;
+	bTracked = PollControllerState(Position, Orientation);
+	if (bTracked)
+	{
+		SetRelativeLocationAndRotation(Position, Orientation);
+	}
+
+	if (!ViewExtension.IsValid() && GEngine)
+	{
+		TSharedPtr< FViewExtension, ESPMode::ThreadSafe > NewViewExtension( new FViewExtension(this) );
+		ViewExtension = NewViewExtension;
+		GEngine->ViewExtensions.Add(ViewExtension);
+	}
+}
 
+//=============================================================================
+bool UMotionControllerComponent::PollControllerState(FVector& Position, FRotator& Orientation)
+{
 	const APlayerController* Actor = Cast<APlayerController>(GetOwner());
 	const bool bHasAuthority = !Actor || Actor->IsLocalPlayerController();
 
@@ -31,10 +84,84 @@ void UMotionControllerComponent::TickComponent(float DeltaTime, enum ELevelTick
 		{
 			if ((MotionController != nullptr) && MotionController->GetControllerOrientationAndPosition(PlayerIndex, Hand, Orientation, Position))
 			{
-				SetRelativeLocationAndRotation(Position, Orientation);
-				bTracked = true;
-				break;
+				return true;
 			}
 		}
 	}
+	return false;
+}
+
+//=============================================================================
+void UMotionControllerComponent::FViewExtension::BeginRenderViewFamily(FSceneViewFamily& InViewFamily)
+{
+	if (!MotionControllerComponent)
+	{
+		return;
+	}
+	FScopeLock ScopeLock(&CritSect);
+	if (!MotionControllerComponent || MotionControllerComponent->bDisableLowLatencyUpdate || !CVarEnableMotionControllerLateUpdate.GetValueOnGameThread())
+	{
+		return;
+	}
+
+	LateUpdateSceneProxyCount = 0;
+	GatherSceneProxies(MotionControllerComponent);
 }
+
+//=============================================================================
+void UMotionControllerComponent::FViewExtension::PreRenderViewFamily_RenderThread(FRHICommandListImmediate& RHICmdList, FSceneViewFamily& InViewFamily)
+{
+	if (!MotionControllerComponent)
+	{
+		return;
+	}
+	FScopeLock ScopeLock(&CritSect);
+	if (!MotionControllerComponent || MotionControllerComponent->bDisableLowLatencyUpdate || !CVarEnableMotionControllerLateUpdate.GetValueOnRenderThread())
+	{
+		return;
+	}
+
+	// Poll state for the most recent controller transform
+	FVector Position;
+	FRotator Orientation;
+	if (!MotionControllerComponent->PollControllerState(Position, Orientation))
+	{
+		return;
+	}
+	
+	// Calculate the late update transform that will rebase all children proxies within the frame of reference
+	const FTransform OldLocalToWorldTransform = MotionControllerComponent->CalcNewComponentToWorld(MotionControllerComponent->GetRelativeTransform());
+	const FTransform NewLocalToWorldTransform = MotionControllerComponent->CalcNewComponentToWorld(FTransform(Orientation, Position));
+	const FMatrix LateUpdateTransform = (OldLocalToWorldTransform.Inverse() * NewLocalToWorldTransform).ToMatrixWithScale();
+
+	// Apply adjustment to the affected scene proxies
+	for (int32 ProxyIndex = 0; ProxyIndex < LateUpdateSceneProxyCount; ++ProxyIndex)
+	{
+		LateUpdateSceneProxies[ProxyIndex]->ApplyLateUpdateTransform(LateUpdateTransform);
+	}
+}
+
+//=============================================================================
+void UMotionControllerComponent::FViewExtension::GatherSceneProxies(USceneComponent* Component)
+{
+	// If a scene proxy is present, cache it
+	UPrimitiveComponent* PrimitiveComponent = dynamic_cast<UPrimitiveComponent*>(Component);
+	if (PrimitiveComponent && PrimitiveComponent->SceneProxy)
+	{
+		check(LateUpdateSceneProxyCount < ARRAY_COUNT(LateUpdateSceneProxies));
+		LateUpdateSceneProxies[LateUpdateSceneProxyCount++] = PrimitiveComponent->SceneProxy;
+	}
+
+	// Gather children proxies
+	const int32 ChildCount = Component->GetNumChildrenComponents();
+	for (int32 ChildIndex = 0; ChildIndex < ChildCount; ++ChildIndex)
+	{
+		USceneComponent* ChildComponent = Component->GetChildComponent(ChildIndex);
+		if (!ChildComponent)
+		{
+			continue;
+		}
+
+		GatherSceneProxies(ChildComponent);
+	}
+}
\ No newline at end of file
diff --git a/Engine/Source/Runtime/HeadMountedDisplay/Public/MotionControllerComponent.h b/Engine/Source/Runtime/HeadMountedDisplay/Public/MotionControllerComponent.h
--- a/Engine/Source/Runtime/HeadMountedDisplay/Public/MotionControllerComponent.h
+++ b/Engine/Source/Runtime/HeadMountedDisplay/Public/MotionControllerComponent.h
@@ -2,6 +2,7 @@
 
 #pragma once
 #include "IMotionController.h"
+#include "SceneViewExtension.h"
 #include "MotionControllerComponent.generated.h"
 
 UCLASS(MinimalAPI, Blueprintable, meta = (BlueprintSpawnableComponent), ClassGroup = MotionController)
@@ -9,6 +10,8 @@ class UMotionControllerComponent : public UPrimitiveComponent
 {
 	GENERATED_UCLASS_BODY()
 
+	~UMotionControllerComponent();
+
 	/** Which player index this motion controller should automatically follow */
 	UPROPERTY(EditAnywhere, BlueprintReadWrite, Category = "MotionController")
 	int32 PlayerIndex;
@@ -17,6 +20,10 @@ class UMotionControllerComponent : public UPrimitiveComponent
 	UPROPERTY(EditAnywhere, BlueprintReadWrite, Category = "MotionController")
 	TEnumAsByte<EControllerHand> Hand;
 
+	/** If false, render transforms within the motion controller hierarchy will be updated a second time immediately before rendering. */
+	UPROPERTY(EditAnywhere, BlueprintReadWrite, Category = "MotionController")
+	uint32 bDisableLowLatencyUpdate:1;
+
 	void TickComponent(float DeltaTime, enum ELevelTick TickType, FActorComponentTickFunction *ThisTickFunction) override;
 
 	/** Whether or not this component had a valid tracked device this frame */
@@ -29,4 +36,35 @@ class UMotionControllerComponent : public UPrimitiveComponent
 private:
 	/** Whether or not this component had a valid tracked controller associated with it this frame*/
 	bool bTracked;
+
+	/** If true, the Position and Orientation args will contain the most recent controller state */
+	bool PollControllerState(FVector& Position, FRotator& Orientation);
+
+	/** View extension object that can persist on the render thread without the motion controller component */
+	class FViewExtension : public ISceneViewExtension, public TSharedFromThis<FViewExtension, ESPMode::ThreadSafe>
+	{
+	public:
+		FViewExtension(UMotionControllerComponent* InMotionControllerComponent) { MotionControllerComponent = InMotionControllerComponent; }
+		virtual ~FViewExtension() {}
+
+		/** ISceneViewExtension interface */
+		virtual void SetupViewFamily(FSceneViewFamily& InViewFamily) override {}
+		virtual void SetupView(FSceneViewFamily& InViewFamily, FSceneView& InView) override {}
+		virtual void BeginRenderViewFamily(FSceneViewFamily& InViewFamily) override;
+		virtual void PreRenderView_RenderThread(FRHICommandListImmediate& RHICmdList, FSceneView& InView) override {}
+		virtual void PreRenderViewFamily_RenderThread(FRHICommandListImmediate& RHICmdList, FSceneViewFamily& InViewFamily) override;
+	private:
+		friend class UMotionControllerComponent;
+
+		/** Walks the component hierarchy gathering scene proxies */
+		void GatherSceneProxies(USceneComponent* Component);
+
+		/** Motion controller component associated with this view extension */
+		UMotionControllerComponent* MotionControllerComponent;
+
+		/** Scene proxies that need late updates */
+		FPrimitiveSceneProxy*	LateUpdateSceneProxies[16];
+		int32					LateUpdateSceneProxyCount;
+	};
+	TSharedPtr< FViewExtension, ESPMode::ThreadSafe > ViewExtension;
 };
