Removing MotionControllerDevices array from GEngine

diff --git a/Engine/Plugins/Runtime/Steam/SteamVR/Source/SteamVRController/Private/SteamVRController.cpp b/Engine/Plugins/Runtime/Steam/SteamVR/Source/SteamVRController/Private/SteamVRController.cpp
--- a/Engine/Plugins/Runtime/Steam/SteamVR/Source/SteamVRController/Private/SteamVRController.cpp
+++ b/Engine/Plugins/Runtime/Steam/SteamVR/Source/SteamVRController/Private/SteamVRController.cpp
@@ -125,14 +125,10 @@ public:
 		Buttons[ (int32)EControllerHand::Right ][ ESteamVRControllerButton::TouchPadRight ] = FGamepadKeyNames::MotionController_Right_FaceButton2;
 
 		IModularFeatures::Get().RegisterModularFeature(GetModularFeatureName(), this);
-
-		//@todo:  fix this.  construction of the controller happens after InitializeMotionControllers(), so we manually insert into the array here.
-		GEngine->MotionControllerDevices.AddUnique(this);
 	}
 
 	virtual ~FSteamVRController()
 	{
-		GEngine->MotionControllerDevices.Remove(this);
 	}
 
 	virtual void Tick( float DeltaTime ) override
diff --git a/Engine/Source/Runtime/Engine/Classes/Engine/Engine.h b/Engine/Source/Runtime/Engine/Classes/Engine/Engine.h
--- a/Engine/Source/Runtime/Engine/Classes/Engine/Engine.h
+++ b/Engine/Source/Runtime/Engine/Classes/Engine/Engine.h
@@ -1509,9 +1509,6 @@ public:
 	/** Extensions that can modify view parameters on the render thread. */
 	TArray<TSharedPtr<class ISceneViewExtension, ESPMode::ThreadSafe> > ViewExtensions;
 
-	/** Reference to the Motion Control devices that are attached, if any */
-	TArray < class IMotionController*> MotionControllerDevices;
-
 	/** Triggered when a world is added. */	
 	DECLARE_EVENT_OneParam( UEngine, FWorldAddedEvent , UWorld* );
 	
@@ -2282,8 +2279,6 @@ protected:
 	 */
 	virtual bool InitializeHMDDevice();
 
-	virtual bool InitializeMotionControllers();
-
 	/**	Record EngineAnalytics information for attached HMD devices. */
 	virtual void RecordHMDAnalytics();
 
diff --git a/Engine/Source/Runtime/Engine/Private/UnrealEngine.cpp b/Engine/Source/Runtime/Engine/Private/UnrealEngine.cpp
--- a/Engine/Source/Runtime/Engine/Private/UnrealEngine.cpp
+++ b/Engine/Source/Runtime/Engine/Private/UnrealEngine.cpp
@@ -764,7 +764,6 @@ void UEngine::Init(IEngineLoop* InEngineLoop)
 
 	// Initialize the HMDs and motion controllers, if any
 	InitializeHMDDevice();
-	InitializeMotionControllers();
 
 	// Disable the screensaver when running the game.
 	if( GIsClient && !GIsEditor )
@@ -2257,17 +2256,6 @@ bool UEngine::InitializeHMDDevice()
 	return StereoRenderingDevice.IsValid();
 }
 
-bool UEngine::InitializeMotionControllers()
-{
-	TArray<IMotionController*> MotionControllers = IModularFeatures::Get().GetModularFeatureImplementations<IMotionController>(IMotionController::GetModularFeatureName());
-	for (auto MotionController : MotionControllers)
-	{
-		MotionControllerDevices.AddUnique(MotionController);
-	}
-
-	return (MotionControllerDevices.Num() > 0);
-}
-
 void UEngine::RecordHMDAnalytics()
 {
 	if( !GIsEditor )
diff --git a/Engine/Source/Runtime/HeadMountedDisplay/Private/MotionControllerComponent.cpp b/Engine/Source/Runtime/HeadMountedDisplay/Private/MotionControllerComponent.cpp
--- a/Engine/Source/Runtime/HeadMountedDisplay/Private/MotionControllerComponent.cpp
+++ b/Engine/Source/Runtime/HeadMountedDisplay/Private/MotionControllerComponent.cpp
@@ -2,6 +2,7 @@
 //
 #include "HeadMountedDisplayPrivate.h"
 #include "MotionControllerComponent.h"
+#include "Features/IModularFeatures.h"
 
 namespace {
 	/** This is to prevent destruction of motion controller components while they are
@@ -80,7 +81,8 @@ bool UMotionControllerComponent::PollControllerState(FVector& Position, FRotator
 
 	if ((PlayerIndex != INDEX_NONE) && bHasAuthority)
 	{
-		for (auto MotionController : GEngine->MotionControllerDevices)
+		TArray<IMotionController*> MotionControllers = IModularFeatures::Get().GetModularFeatureImplementations<IMotionController>( IMotionController::GetModularFeatureName() );
+		for (auto MotionController : MotionControllers)
 		{
 			if ((MotionController != nullptr) && MotionController->GetControllerOrientationAndPosition(PlayerIndex, Hand, Orientation, Position))
 			{
