#ue4 Motion controller late update fixes:

-MotionControllerComponent now caches "bHasAuthority" state from the game thread for reference on the render thread
-SteamVRHMD code now allows position/orientation sampling on the render thread

diff --git a/Engine/Plugins/Runtime/Steam/SteamVR/Source/SteamVR/Private/SteamVRHMD.cpp b/Engine/Plugins/Runtime/Steam/SteamVR/Source/SteamVR/Private/SteamVRHMD.cpp
--- a/Engine/Plugins/Runtime/Steam/SteamVR/Source/SteamVR/Private/SteamVRHMD.cpp
+++ b/Engine/Plugins/Runtime/Steam/SteamVR/Source/SteamVR/Private/SteamVRHMD.cpp
@@ -411,8 +411,6 @@ void FSteamVRHMD::GetTrackedDeviceIds(ESteamVRTrackedDeviceType DeviceType, TArr
 
 bool FSteamVRHMD::GetTrackedObjectOrientationAndPosition(uint32 DeviceId, FQuat& CurrentOrientation, FVector& CurrentPosition)
 {
-	check(IsInGameThread());
-
 	bool bHasValidPose = false;
 
 	if (DeviceId >= 0 && DeviceId < vr::k_unMaxTrackedDeviceCount)
@@ -429,8 +427,6 @@ bool FSteamVRHMD::GetTrackedObjectOrientationAndPosition(uint32 DeviceId, FQuat&
 
 bool FSteamVRHMD::GetControllerHandPositionAndOrientation( const int32 ControllerIndex, EControllerHand Hand, FVector& OutPosition, FQuat& OutOrientation )
 {
-	check( IsInGameThread() );
-
 	if ((ControllerIndex < 0) || (ControllerIndex >= MAX_STEAMVR_CONTROLLER_PAIRS) || Hand < EControllerHand::Left || Hand > EControllerHand::Right)
 	{
 		return false;
diff --git a/Engine/Source/Runtime/HeadMountedDisplay/Private/MotionControllerComponent.cpp b/Engine/Source/Runtime/HeadMountedDisplay/Private/MotionControllerComponent.cpp
--- a/Engine/Source/Runtime/HeadMountedDisplay/Private/MotionControllerComponent.cpp
+++ b/Engine/Source/Runtime/HeadMountedDisplay/Private/MotionControllerComponent.cpp
@@ -29,6 +29,7 @@ UMotionControllerComponent::UMotionControllerComponent(const FObjectInitializer&
 	PlayerIndex = 0;
 	Hand = EControllerHand::Left;
 	bDisableLowLatencyUpdate = false;
+	bHasAuthority = false;
 }
 
 //=============================================================================
@@ -75,8 +76,12 @@ void UMotionControllerComponent::TickComponent(float DeltaTime, enum ELevelTick
 //=============================================================================
 bool UMotionControllerComponent::PollControllerState(FVector& Position, FRotator& Orientation)
 {
-	const APlayerController* Actor = Cast<APlayerController>(GetOwner());
-	const bool bHasAuthority = !Actor || Actor->IsLocalPlayerController();
+	if (IsInGameThread())
+	{
+		// Cache state from the game thread for use on the render thread
+		const APlayerController* Actor = Cast<APlayerController>(GetOwner());
+		bHasAuthority = !Actor || Actor->IsLocalPlayerController();
+	}
 
 	if ((PlayerIndex != INDEX_NONE) && bHasAuthority)
 	{
diff --git a/Engine/Source/Runtime/HeadMountedDisplay/Public/MotionControllerComponent.h b/Engine/Source/Runtime/HeadMountedDisplay/Public/MotionControllerComponent.h
--- a/Engine/Source/Runtime/HeadMountedDisplay/Public/MotionControllerComponent.h
+++ b/Engine/Source/Runtime/HeadMountedDisplay/Public/MotionControllerComponent.h
@@ -37,6 +37,9 @@ private:
 	/** Whether or not this component had a valid tracked controller associated with it this frame*/
 	bool bTracked;
 
+	/** Whether or not this component has authority within the frame*/
+	bool bHasAuthority;
+
 	/** If true, the Position and Orientation args will contain the most recent controller state */
 	bool PollControllerState(FVector& Position, FRotator& Orientation);
 
