UE-49758 vr.HiddenAreaMask should default to 1 (On)
vr.HiddenAreaMask default back to 1, was 0 for a specific demo and accidentally got into main.

diff --git a/Engine/Source/Runtime/HeadMountedDisplay/Private/HeadMountedDisplayConsoleCommands.cpp b/Engine/Source/Runtime/HeadMountedDisplay/Private/HeadMountedDisplayConsoleCommands.cpp
--- a/Engine/Source/Runtime/HeadMountedDisplay/Private/HeadMountedDisplayConsoleCommands.cpp
+++ b/Engine/Source/Runtime/HeadMountedDisplay/Private/HeadMountedDisplayConsoleCommands.cpp
@@ -25,7 +25,7 @@ DEFINE_LOG_CATEGORY_STATIC(LogHeadMountedDisplayCommand, Display, All);
 */
 static TAutoConsoleVariable<int32> CVarHiddenAreaMask(
 	TEXT("vr.HiddenAreaMask"),
-	0,
+	1,
 	*LOCTEXT("CVarText_HiddenAreaMask", "Enable or disable hidden area mask\n0: disabled\n1: enabled").ToString(),
 	ECVF_Scalability | ECVF_RenderThreadSafe);
 
