Fix for HMD errors

diff --git a/Engine/Plugins/Runtime/Steam/SteamVR/Source/SteamVR/Private/SteamVRHMD.h b/Engine/Plugins/Runtime/Steam/SteamVR/Source/SteamVR/Private/SteamVRHMD.h
--- a/Engine/Plugins/Runtime/Steam/SteamVR/Source/SteamVR/Private/SteamVRHMD.h
+++ b/Engine/Plugins/Runtime/Steam/SteamVR/Source/SteamVR/Private/SteamVRHMD.h
@@ -13,6 +13,7 @@
 #include "StereoLayerManager.h"
 #include "XRRenderTargetManager.h"
 #include "XRRenderBridge.h"
+#include "IHeadMountedDisplayVulkanExtensions.h"
 
 #if PLATFORM_WINDOWS
 #include "Windows/AllowWindowsPlatformTypes.h"
diff --git a/Engine/Source/Runtime/HeadMountedDisplay/Public/IHeadMountedDisplayModule.h b/Engine/Source/Runtime/HeadMountedDisplay/Public/IHeadMountedDisplayModule.h
--- a/Engine/Source/Runtime/HeadMountedDisplay/Public/IHeadMountedDisplayModule.h
+++ b/Engine/Source/Runtime/HeadMountedDisplay/Public/IHeadMountedDisplayModule.h
@@ -8,8 +8,8 @@
 #include "Features/IModularFeatures.h"
 #include "Features/IModularFeature.h"
 
-#include "IHeadMountedDisplayVulkanExtensions.h"
-
+//#include "IHeadMountedDisplayVulkanExtensions.h"
+class IHeadMountedDisplayVulkanExtensions;
 class IXRTrackingSystem;
 
 /**
