For MRC calibration, ensuring the alignment models get properly positioned (not clipping in front of the near plane)
  - Forcing the size-to-fit adjustment to fit things in front of the near plane
  - Ensuring all alignment points that are supposed to be at the same depth plane, remain so
  - Making sure we size-to-fit after the model has been async loaded (ensuring it has the proper bounds, which we're fitting).

diff --git a/Engine/Plugins/Runtime/Steam/SteamVR/Source/SteamVR/Private/SteamVRAssetManager.cpp b/Engine/Plugins/Runtime/Steam/SteamVR/Source/SteamVR/Private/SteamVRAssetManager.cpp
--- a/Engine/Plugins/Runtime/Steam/SteamVR/Source/SteamVR/Private/SteamVRAssetManager.cpp
+++ b/Engine/Plugins/Runtime/Steam/SteamVR/Source/SteamVR/Private/SteamVRAssetManager.cpp
@@ -833,8 +833,18 @@ UPrimitiveComponent* FSteamVRAssetManager::CreateRenderComponent(const int32 Dev
 				AssignedMeshLoader.Pin()->Tick(0.0f);
 			}
 		}
+		else
+		{
+			// failure...
+			OnLoadComplete.ExecuteIfBound(nullptr);
+		}
 	}
+	else
 #endif
+	{
+		// failure...
+		OnLoadComplete.ExecuteIfBound(nullptr);
+	}
 	return NewRenderComponent;
 }
 
diff --git a/Engine/Source/Runtime/HeadMountedDisplay/Private/HeadMountedDisplayFunctionLibrary.cpp b/Engine/Source/Runtime/HeadMountedDisplay/Private/HeadMountedDisplayFunctionLibrary.cpp
--- a/Engine/Source/Runtime/HeadMountedDisplay/Private/HeadMountedDisplayFunctionLibrary.cpp
+++ b/Engine/Source/Runtime/HeadMountedDisplay/Private/HeadMountedDisplayFunctionLibrary.cpp
@@ -384,3 +384,20 @@ void UHeadMountedDisplayFunctionLibrary::GetDeviceWorldPose(UObject* WorldContex
 	FQuat WorldOrientation = TrackingToWorld.TransformRotation(Orientation.Quaternion());
 	Orientation = WorldOrientation.Rotator();
 }
+
+bool UHeadMountedDisplayFunctionLibrary::IsDeviceTracking(const FXRDeviceId& XRDeviceId)
+{
+	bool bIsTracked = false;
+
+	// @TODO: It seems certain IXRTrackingSystem's aren't registering themselves with the modular feature framework. Ideally we'd be loop over them instead of picking just one.
+	IXRTrackingSystem* TrackingSys = GEngine->XRSystem.Get();
+	if (TrackingSys)
+	{
+		if (XRDeviceId.IsOwnedBy(TrackingSys))
+		{
+			bIsTracked = TrackingSys->IsTracking(XRDeviceId.DeviceId);
+		}
+	}
+
+	return bIsTracked;
+}
diff --git a/Engine/Source/Runtime/HeadMountedDisplay/Private/MotionTrackedDeviceFunctionLibrary.cpp b/Engine/Source/Runtime/HeadMountedDisplay/Private/MotionTrackedDeviceFunctionLibrary.cpp
--- a/Engine/Source/Runtime/HeadMountedDisplay/Private/MotionTrackedDeviceFunctionLibrary.cpp
+++ b/Engine/Source/Runtime/HeadMountedDisplay/Private/MotionTrackedDeviceFunctionLibrary.cpp
@@ -217,3 +217,34 @@ FName UMotionTrackedDeviceFunctionLibrary::GetActiveTrackingSystemName()
 	IXRTrackingSystem* TrackingSys = GEngine->XRSystem.Get();
 	return TrackingSys ? TrackingSys->GetSystemName() : NAME_None;
 }
+
+bool UMotionTrackedDeviceFunctionLibrary::IsMotionSourceTracking(int32 PlayerIndex, FName SourceName)
+{
+	bool bIsTracking = false;
+
+	if (SourceName == FXRMotionControllerBase::HMDSourceId)
+	{
+		IXRTrackingSystem* TrackingSys = GEngine->XRSystem.Get();
+		if (TrackingSys)
+		{
+			bIsTracking = TrackingSys->IsTracking(IXRTrackingSystem::HMDDeviceId);
+		}
+	}
+	else
+	{
+		TArray<IMotionController*> MotionControllers = IModularFeatures::Get().GetModularFeatureImplementations<IMotionController>(IMotionController::GetModularFeatureName());
+		for (IMotionController* MotionController : MotionControllers)
+		{
+			if (MotionController)
+			{
+				if (MotionController->GetControllerTrackingStatus(PlayerIndex, SourceName) != ETrackingStatus::NotTracked)
+				{
+					bIsTracking = true;
+					break;
+				}
+			}
+		}
+	}
+
+	return bIsTracking;
+}
diff --git a/Engine/Source/Runtime/HeadMountedDisplay/Private/XRAssetFunctionLibrary.cpp b/Engine/Source/Runtime/HeadMountedDisplay/Private/XRAssetFunctionLibrary.cpp
--- a/Engine/Source/Runtime/HeadMountedDisplay/Private/XRAssetFunctionLibrary.cpp
+++ b/Engine/Source/Runtime/HeadMountedDisplay/Private/XRAssetFunctionLibrary.cpp
@@ -13,7 +13,7 @@
 namespace XRAssetFunctionLibrary_Impl
 {
 	static UPrimitiveComponent* AddDeviceVisualizationComponent(AActor* Target, const FXRDeviceId& XRDeviceId, bool bManualAttachment, const FTransform& RelativeTransformbool, const bool bForceSynchronous = false, const FXRComponentLoadComplete& OnLoadComplete = FXRComponentLoadComplete());
-	static UPrimitiveComponent* AddNamedDeviceVisualizationComponent(AActor* Target, const FName SystemName, const FName DeviceName, bool bManualAttachment, const FTransform& RelativeTransform, const bool bForceSynchronous = false, const FXRComponentLoadComplete& OnLoadComplete = FXRComponentLoadComplete());
+	static UPrimitiveComponent* AddNamedDeviceVisualizationComponent(AActor* Target, const FName SystemName, const FName DeviceName, bool bManualAttachment, const FTransform& RelativeTransform, FXRDeviceId& XRDeviceIdOut, const bool bForceSynchronous = false, const FXRComponentLoadComplete& OnLoadComplete = FXRComponentLoadComplete());
 	static UPrimitiveComponent* SpawnDeviceComponent(AActor* Target, IXRSystemAssets* AssetSystem, int32 DeviceId, bool bManualAttachment, const FTransform& RelativeTransform, const bool bForceSynchronous, const FXRComponentLoadComplete& OnLoadComplete);
 }
 
@@ -62,7 +62,8 @@ static UPrimitiveComponent* XRAssetFunctionLibrary_Impl::AddNamedDeviceVisualiza
 	const FName						SystemName, 
 	const FName						DeviceName, 
 	bool							bManualAttachment, 
-	const FTransform&				RelativeTransform, 
+	const FTransform&				RelativeTransform,
+	FXRDeviceId&					XRDeviceIdOut,
 	const bool						bForceSynchronous, 
 	const FXRComponentLoadComplete& OnLoadComplete)
 {
@@ -81,6 +82,7 @@ static UPrimitiveComponent* XRAssetFunctionLibrary_Impl::AddNamedDeviceVisualiza
 				NewComponent = SpawnDeviceComponent(Target, AssetSys, DeviceId, bManualAttachment, RelativeTransform, bForceSynchronous, OnLoadComplete);
 				if (NewComponent)
 				{
+					XRDeviceIdOut = FXRDeviceId(AssetSys, DeviceId);
 					break;
 				}
 			}
@@ -89,7 +91,7 @@ static UPrimitiveComponent* XRAssetFunctionLibrary_Impl::AddNamedDeviceVisualiza
 
 	if (NewComponent == nullptr)
 	{
-		UE_LOG(LogHMD, Warning, TEXT("Failed to find an active XR system with a model for the requested device."));
+		UE_LOG(LogHMD, Warning, TEXT("Failed to find an active XR system with a model for the requested device: %s."), *DeviceName.ToString());
 	}
 	return NewComponent;
 }
@@ -145,9 +147,9 @@ UPrimitiveComponent* UXRAssetFunctionLibrary::AddDeviceVisualizationComponentBlo
 	return NewComponent;
 }
 
-UPrimitiveComponent* UXRAssetFunctionLibrary::AddNamedDeviceVisualizationComponentBlocking(AActor* Target, const FName SystemName, const FName DeviceName, bool bManualAttachment, const FTransform& RelativeTransform)
+UPrimitiveComponent* UXRAssetFunctionLibrary::AddNamedDeviceVisualizationComponentBlocking(AActor* Target, const FName SystemName, const FName DeviceName, bool bManualAttachment, const FTransform& RelativeTransform, FXRDeviceId& XRDeviceId)
 {
-	UPrimitiveComponent* NewComponent = XRAssetFunctionLibrary_Impl::AddNamedDeviceVisualizationComponent(Target, SystemName, DeviceName, bManualAttachment, RelativeTransform, /*bForceSynchronous =*/true);
+	UPrimitiveComponent* NewComponent = XRAssetFunctionLibrary_Impl::AddNamedDeviceVisualizationComponent(Target, SystemName, DeviceName, bManualAttachment, RelativeTransform, XRDeviceId, /*bForceSynchronous =*/true);
 	return NewComponent;
 }
 
@@ -175,7 +177,7 @@ UAsyncTask_LoadXRDeviceVisComponent* UAsyncTask_LoadXRDeviceVisComponent::AddDev
 	return NewTask;
 }
 
-UAsyncTask_LoadXRDeviceVisComponent* UAsyncTask_LoadXRDeviceVisComponent::AddNamedDeviceVisualizationComponentAsync(AActor* Target, const FName SystemName, const FName DeviceName, bool bManualAttachment, const FTransform& RelativeTransform, UPrimitiveComponent*& NewComponent)
+UAsyncTask_LoadXRDeviceVisComponent* UAsyncTask_LoadXRDeviceVisComponent::AddNamedDeviceVisualizationComponentAsync(AActor* Target, const FName SystemName, const FName DeviceName, bool bManualAttachment, const FTransform& RelativeTransform, FXRDeviceId& XRDeviceId, UPrimitiveComponent*& NewComponent)
 {
 	UAsyncTask_LoadXRDeviceVisComponent* NewTask = NewObject<UAsyncTask_LoadXRDeviceVisComponent>();
 
@@ -187,7 +189,7 @@ UAsyncTask_LoadXRDeviceVisComponent* UAsyncTask_LoadXRDeviceVisComponent::AddNam
 			NewTask->OnLoadComplete(Component != nullptr);
 		}
 	});
-	NewComponent = NewTask->SpawnedComponent = XRAssetFunctionLibrary_Impl::AddNamedDeviceVisualizationComponent(Target, SystemName, DeviceName, bManualAttachment, RelativeTransform, /*bForceSynchronous =*/false, LoadCompleteDelegate);
+	NewComponent = NewTask->SpawnedComponent = XRAssetFunctionLibrary_Impl::AddNamedDeviceVisualizationComponent(Target, SystemName, DeviceName, bManualAttachment, RelativeTransform, XRDeviceId, /*bForceSynchronous =*/false, LoadCompleteDelegate);
 
 	return NewTask;
 }
diff --git a/Engine/Source/Runtime/HeadMountedDisplay/Public/HeadMountedDisplayFunctionLibrary.h b/Engine/Source/Runtime/HeadMountedDisplay/Public/HeadMountedDisplayFunctionLibrary.h
--- a/Engine/Source/Runtime/HeadMountedDisplay/Public/HeadMountedDisplayFunctionLibrary.h
+++ b/Engine/Source/Runtime/HeadMountedDisplay/Public/HeadMountedDisplayFunctionLibrary.h
@@ -271,4 +271,12 @@ class HEADMOUNTEDDISPLAY_API UHeadMountedDisplayFunctionLibrary : public UBluepr
 	 */
 	UFUNCTION(BlueprintCallable, Category="Input|XRTracking",  meta=(WorldContext="WorldContext"))
 	static void GetDeviceWorldPose(UObject* WorldContext, const FXRDeviceId& XRDeviceId, bool& bIsTracked, FRotator& Orientation, bool& bHasPositionalTracking, FVector& Position);
+
+	/**
+	 * Cross XR-System query that returns whether the specified device is tracked or not.
+	 *
+	 * @param  XRDeviceId	Specifies the device you're querying for.
+	 */
+	UFUNCTION(BlueprintCallable, Category="Input|XRTracking")
+	static bool IsDeviceTracking(const FXRDeviceId& XRDeviceId);
 };
diff --git a/Engine/Source/Runtime/HeadMountedDisplay/Public/MotionTrackedDeviceFunctionLibrary.h b/Engine/Source/Runtime/HeadMountedDisplay/Public/MotionTrackedDeviceFunctionLibrary.h
--- a/Engine/Source/Runtime/HeadMountedDisplay/Public/MotionTrackedDeviceFunctionLibrary.h
+++ b/Engine/Source/Runtime/HeadMountedDisplay/Public/MotionTrackedDeviceFunctionLibrary.h
@@ -166,4 +166,12 @@ class HEADMOUNTEDDISPLAY_API UMotionTrackedDeviceFunctionLibrary : public UBluep
 	 */
 	UFUNCTION(BlueprintCallable, Category = "Input|MotionTracking")
 	static FName GetActiveTrackingSystemName();
+
+	/**
+	 * Queries the specified source's tracking status and returns true if it has tracking.
+	 *
+	 * @return Tracking status of the specified controller.
+	 */
+	UFUNCTION(BlueprintCallable, Category = "Input|MotionTracking")
+	static bool IsMotionSourceTracking(int32 PlayerIndex, FName SourceName);
 };
diff --git a/Engine/Source/Runtime/HeadMountedDisplay/Public/XRAssetFunctionLibrary.h b/Engine/Source/Runtime/HeadMountedDisplay/Public/XRAssetFunctionLibrary.h
--- a/Engine/Source/Runtime/HeadMountedDisplay/Public/XRAssetFunctionLibrary.h
+++ b/Engine/Source/Runtime/HeadMountedDisplay/Public/XRAssetFunctionLibrary.h
@@ -55,7 +55,7 @@ public:
 	 * @return A new component representing the specified device (invalid/null if a model for the device doesn't exist).
 	 */
 	UFUNCTION(BlueprintCallable, Category = "XR|Devices", meta = (DefaultToSelf = "Target"))
-	static UPrimitiveComponent* AddNamedDeviceVisualizationComponentBlocking(AActor* Target, const FName SystemName, const FName DeviceName, bool bManualAttachment, const FTransform& RelativeTransform);
+	static UPrimitiveComponent* AddNamedDeviceVisualizationComponentBlocking(AActor* Target, const FName SystemName, const FName DeviceName, bool bManualAttachment, const FTransform& RelativeTransform, FXRDeviceId& XRDeviceId);
 };
 
 /* UAsyncTask_LoadXRDeviceVisComponent
@@ -71,7 +71,7 @@ class UAsyncTask_LoadXRDeviceVisComponent : public UBlueprintAsyncActionBase
 public:
 
 	UFUNCTION(BlueprintCallable, meta = (BlueprintInternalUseOnly = "true", DefaultToSelf = "Target"), Category = "XR|Devices")
-	static UAsyncTask_LoadXRDeviceVisComponent* AddNamedDeviceVisualizationComponentAsync(AActor* Target, const FName SystemName, const FName DeviceName, bool bManualAttachment, const FTransform& RelativeTransform, UPrimitiveComponent*& NewComponent);
+	static UAsyncTask_LoadXRDeviceVisComponent* AddNamedDeviceVisualizationComponentAsync(AActor* Target, const FName SystemName, const FName DeviceName, bool bManualAttachment, const FTransform& RelativeTransform, FXRDeviceId& XRDeviceId, UPrimitiveComponent*& NewComponent);
 
 	UFUNCTION(BlueprintCallable, meta = (BlueprintInternalUseOnly = "true", DefaultToSelf = "Target"), Category = "XR|Devices")
 	static UAsyncTask_LoadXRDeviceVisComponent* AddDeviceVisualizationComponentAsync(AActor* Target, const FXRDeviceId& XRDeviceId, bool bManualAttachment, const FTransform& RelativeTransform, UPrimitiveComponent*& NewComponent);
