Duplicating for 4.20 cl 4184935
UEVR-1241 Gesture keypoint drift
-If the render state is recreated SendRenderTransform_Concurrent doesn't happen.  We need to also update the cached old transform in CreateRenderState_Concurrent

diff --git a/Engine/Source/Runtime/HeadMountedDisplay/Private/MotionControllerComponent.cpp b/Engine/Source/Runtime/HeadMountedDisplay/Private/MotionControllerComponent.cpp
--- a/Engine/Source/Runtime/HeadMountedDisplay/Private/MotionControllerComponent.cpp
+++ b/Engine/Source/Runtime/HeadMountedDisplay/Private/MotionControllerComponent.cpp
@@ -92,6 +92,13 @@ void UMotionControllerComponent::BeginDestroy()
 	}
 }
 
+void UMotionControllerComponent::CreateRenderState_Concurrent()
+{
+	Super::CreateRenderState_Concurrent();
+	RenderThreadRelativeTransform = GetRelativeTransform();
+	RenderThreadComponentScale = GetComponentScale();
+}
+
 void UMotionControllerComponent::SendRenderTransform_Concurrent()
 {
 	RenderThreadRelativeTransform = GetRelativeTransform();
diff --git a/Engine/Source/Runtime/HeadMountedDisplay/Public/MotionControllerComponent.h b/Engine/Source/Runtime/HeadMountedDisplay/Public/MotionControllerComponent.h
--- a/Engine/Source/Runtime/HeadMountedDisplay/Public/MotionControllerComponent.h
+++ b/Engine/Source/Runtime/HeadMountedDisplay/Public/MotionControllerComponent.h
@@ -106,6 +106,7 @@ public:
 
 protected:
 	//~ Begin UActorComponent Interface.
+	virtual void CreateRenderState_Concurrent() override;
 	virtual void SendRenderTransform_Concurrent() override;
 	//~ End UActorComponent Interface.
 
