Merging using Dev-VR_to_Release-4.21

diff --git a/Engine/Source/Runtime/HeadMountedDisplay/Private/DefaultXRCamera.cpp b/Engine/Source/Runtime/HeadMountedDisplay/Private/DefaultXRCamera.cpp
--- a/Engine/Source/Runtime/HeadMountedDisplay/Private/DefaultXRCamera.cpp
+++ b/Engine/Source/Runtime/HeadMountedDisplay/Private/DefaultXRCamera.cpp
@@ -191,11 +191,12 @@ void FDefaultXRCamera::PostRenderViewFamily_RenderThread(FRHICommandListImmediat
 void FDefaultXRCamera::SetupViewFamily(FSceneViewFamily& InViewFamily)
 {
 	static const auto CVarAllowMotionBlurInVR = IConsoleManager::Get().FindTConsoleVariableDataInt(TEXT("vr.AllowMotionBlurInVR"));
-	const bool allowMotionBlur = (CVarAllowMotionBlurInVR && CVarAllowMotionBlurInVR->GetValueOnAnyThread() != 0);
+	const bool AllowMotionBlur = (CVarAllowMotionBlurInVR && CVarAllowMotionBlurInVR->GetValueOnAnyThread() != 0);
 	const IHeadMountedDisplay* HMD = TrackingSystem->GetHMDDevice();
-	InViewFamily.EngineShowFlags.MotionBlur = allowMotionBlur;
-	InViewFamily.EngineShowFlags.HMDDistortion = HMD?HMD->GetHMDDistortionEnabled(InViewFamily.Scene->GetShadingPath()):false;
+	InViewFamily.EngineShowFlags.MotionBlur = AllowMotionBlur;
+	InViewFamily.EngineShowFlags.HMDDistortion = HMD != nullptr ? HMD->GetHMDDistortionEnabled(InViewFamily.Scene->GetShadingPath()) : false;
 	InViewFamily.EngineShowFlags.StereoRendering = bCurrentFrameIsStereoRendering;
+	InViewFamily.EngineShowFlags.Rendering = HMD != nullptr ? !HMD->IsRenderingPaused() : false;
 }
 
 void FDefaultXRCamera::SetupView(FSceneViewFamily& InViewFamily, FSceneView& InView)
