Merging using Dev-VR_to_Release-4.21

diff --git a/Engine/Source/Runtime/HeadMountedDisplay/Private/DefaultXRCamera.cpp b/Engine/Source/Runtime/HeadMountedDisplay/Private/DefaultXRCamera.cpp
--- a/Engine/Source/Runtime/HeadMountedDisplay/Private/DefaultXRCamera.cpp
+++ b/Engine/Source/Runtime/HeadMountedDisplay/Private/DefaultXRCamera.cpp
@@ -196,7 +196,7 @@ void FDefaultXRCamera::SetupViewFamily(FSceneViewFamily& InViewFamily)
 	InViewFamily.EngineShowFlags.MotionBlur = AllowMotionBlur;
 	InViewFamily.EngineShowFlags.HMDDistortion = HMD != nullptr ? HMD->GetHMDDistortionEnabled(InViewFamily.Scene->GetShadingPath()) : false;
 	InViewFamily.EngineShowFlags.StereoRendering = bCurrentFrameIsStereoRendering;
-	InViewFamily.EngineShowFlags.Rendering = HMD != nullptr ? !HMD->IsRenderingPaused() : false;
+	InViewFamily.EngineShowFlags.Rendering = HMD != nullptr ? !HMD->IsRenderingPaused() : true;
 }
 
 void FDefaultXRCamera::SetupView(FSceneViewFamily& InViewFamily, FSceneView& InView)
