Added a flag bUseAlpha that is checked in FDefaultSpectatorScreenController::RenderSpectatorModeMirrorAndTexture(), if (SpectatorScreenModeTexturePlusEyeLayout_RenderThread.bDrawEyeFirst).  It can be set with the SetSpectatorScreenModeTexturePlusEyeLayout Blueprint node.

diff --git a/Engine/Source/Runtime/HeadMountedDisplay/Private/DefaultSpectatorScreenController.cpp b/Engine/Source/Runtime/HeadMountedDisplay/Private/DefaultSpectatorScreenController.cpp
--- a/Engine/Source/Runtime/HeadMountedDisplay/Private/DefaultSpectatorScreenController.cpp
+++ b/Engine/Source/Runtime/HeadMountedDisplay/Private/DefaultSpectatorScreenController.cpp
@@ -351,7 +351,7 @@ void FDefaultSpectatorScreenController::RenderSpectatorModeMirrorAndTexture(FRHI
 	if (SpectatorScreenModeTexturePlusEyeLayout_RenderThread.bDrawEyeFirst)
 	{
 		HMDDevice->CopyTexture_RenderThread(RHICmdList, EyeTexture, CroppedEyeSrcRect, TargetTexture, EyeDstRect, bClearBlack, true);
-		HMDDevice->CopyTexture_RenderThread(RHICmdList, OtherTextureLocal, OtherSrcRect, TargetTexture, OtherDstRect, false, true);
+		HMDDevice->CopyTexture_RenderThread(RHICmdList, OtherTextureLocal, OtherSrcRect, TargetTexture, OtherDstRect, false, !SpectatorScreenModeTexturePlusEyeLayout_RenderThread.bUseAlpha);
 	}
 	else
 	{
diff --git a/Engine/Source/Runtime/HeadMountedDisplay/Private/HeadMountedDisplayFunctionLibrary.cpp b/Engine/Source/Runtime/HeadMountedDisplay/Private/HeadMountedDisplayFunctionLibrary.cpp
--- a/Engine/Source/Runtime/HeadMountedDisplay/Private/HeadMountedDisplayFunctionLibrary.cpp
+++ b/Engine/Source/Runtime/HeadMountedDisplay/Private/HeadMountedDisplayFunctionLibrary.cpp
@@ -331,12 +331,12 @@ void UHeadMountedDisplayFunctionLibrary::SetSpectatorScreenTexture(UTexture* InT
 	}
 }
 
-void UHeadMountedDisplayFunctionLibrary::SetSpectatorScreenModeTexturePlusEyeLayout(FVector2D EyeRectMin, FVector2D EyeRectMax, FVector2D TextureRectMin, FVector2D TextureRectMax, bool bDrawEyeFirst /* = true */, bool bClearBlack /* = false */)
+void UHeadMountedDisplayFunctionLibrary::SetSpectatorScreenModeTexturePlusEyeLayout(FVector2D EyeRectMin, FVector2D EyeRectMax, FVector2D TextureRectMin, FVector2D TextureRectMax, bool bDrawEyeFirst /* = true */, bool bClearBlack /* = false */, bool bUseAlpha /* = false */)
 {
 	ISpectatorScreenController* const Controller = HMDFunctionLibraryHelpers::GetSpectatorScreenController();
 	if (Controller)
 	{
-		Controller->SetSpectatorScreenModeTexturePlusEyeLayout(FSpectatorScreenModeTexturePlusEyeLayout(EyeRectMin, EyeRectMax, TextureRectMin, TextureRectMax, bDrawEyeFirst, bClearBlack));
+		Controller->SetSpectatorScreenModeTexturePlusEyeLayout(FSpectatorScreenModeTexturePlusEyeLayout(EyeRectMin, EyeRectMax, TextureRectMin, TextureRectMax, bDrawEyeFirst, bClearBlack, bUseAlpha));
 	}
 }
 
diff --git a/Engine/Source/Runtime/HeadMountedDisplay/Public/HeadMountedDisplayFunctionLibrary.h b/Engine/Source/Runtime/HeadMountedDisplay/Public/HeadMountedDisplayFunctionLibrary.h
--- a/Engine/Source/Runtime/HeadMountedDisplay/Public/HeadMountedDisplayFunctionLibrary.h
+++ b/Engine/Source/Runtime/HeadMountedDisplay/Public/HeadMountedDisplayFunctionLibrary.h
@@ -258,7 +258,7 @@ class HEADMOUNTEDDISPLAY_API UHeadMountedDisplayFunctionLibrary : public UBluepr
 	* @param	bClearBlack: if true the render target will be drawn black before either rect is drawn.
 	*/
 	UFUNCTION(BlueprintCallable, Category = "Input|HeadMountedDisplay|SpectatorScreen")
-	static void SetSpectatorScreenModeTexturePlusEyeLayout(FVector2D EyeRectMin, FVector2D EyeRectMax, FVector2D TextureRectMin, FVector2D TextureRectMax, bool bDrawEyeFirst = true, bool bClearBlack = false);
+	static void SetSpectatorScreenModeTexturePlusEyeLayout(FVector2D EyeRectMin, FVector2D EyeRectMax, FVector2D TextureRectMin, FVector2D TextureRectMax, bool bDrawEyeFirst = true, bool bClearBlack = false, bool bUseAlpha = false);
 
 	/**
 	 * Cross XR-System query that will list all XR devices currently being tracked.
diff --git a/Engine/Source/Runtime/HeadMountedDisplay/Public/HeadMountedDisplayTypes.h b/Engine/Source/Runtime/HeadMountedDisplay/Public/HeadMountedDisplayTypes.h
--- a/Engine/Source/Runtime/HeadMountedDisplay/Public/HeadMountedDisplayTypes.h
+++ b/Engine/Source/Runtime/HeadMountedDisplay/Public/HeadMountedDisplayTypes.h
@@ -108,15 +108,17 @@ struct FSpectatorScreenModeTexturePlusEyeLayout
 		, TextureRectMax(0.25f, 0.25f)
 		, bDrawEyeFirst(true)
 		, bClearBlack(false)
+		, bUseAlpha(false)
 	{}
 
-	FSpectatorScreenModeTexturePlusEyeLayout(FVector2D InEyeRectMin, FVector2D InEyeRectMax, FVector2D InTextureRectMin, FVector2D InTextureRectMax, bool InbDrawEyeFirst, bool InbClearBlack)
+	FSpectatorScreenModeTexturePlusEyeLayout(FVector2D InEyeRectMin, FVector2D InEyeRectMax, FVector2D InTextureRectMin, FVector2D InTextureRectMax, bool InbDrawEyeFirst, bool InbClearBlack, bool InbUseAlpha)
 		: EyeRectMin(InEyeRectMin)
 		, EyeRectMax(InEyeRectMax)
 		, TextureRectMin(InTextureRectMin)
 		, TextureRectMax(InTextureRectMax)
 		, bDrawEyeFirst(InbDrawEyeFirst)
 		, bClearBlack(InbClearBlack)
+		, bUseAlpha(InbUseAlpha)
 	{}
 
 	bool IsValid() const 
@@ -169,6 +171,7 @@ struct FSpectatorScreenModeTexturePlusEyeLayout
 	FVector2D TextureRectMin;
 	FVector2D TextureRectMax;
 	bool bDrawEyeFirst;
+	bool bUseAlpha;
 	bool bClearBlack;
 };
 
