Adding in support for locking HMD tracking to an external tracking environment (e.g. mo cap).  This is not meant to drive the position continually in lieu of an HMD's built in system, but is rather to help keep the HMD from generally getting out of alignment with another tracking system.  Two functions, CalibrateExternalTrackingToHMD and UpdateExternalTrackingHMDPosition allow for arbitrary rigid offsets between the HMD and the external tracking source, and updating as desired.

diff --git a/Engine/Plugins/Runtime/Steam/SteamVR/Source/SteamVR/Private/SteamVRHMD.cpp b/Engine/Plugins/Runtime/Steam/SteamVR/Source/SteamVR/Private/SteamVRHMD.cpp
--- a/Engine/Plugins/Runtime/Steam/SteamVR/Source/SteamVR/Private/SteamVRHMD.cpp
+++ b/Engine/Plugins/Runtime/Steam/SteamVR/Source/SteamVR/Private/SteamVRHMD.cpp
@@ -908,7 +908,7 @@ void FSteamVRHMD::PoseToOrientationAndPosition(const vr::HmdMatrix34_t& InPose,
 	OutOrientation.Z = Orientation.Y;
 	OutOrientation.W = -Orientation.W;
 
-	FVector Position = ((FVector(-Pose.M[3][2], Pose.M[3][0], Pose.M[3][1]) - BaseOffset) * WorldToMetersScale);
+	FVector Position = ((FVector(-Pose.M[3][2], Pose.M[3][0], Pose.M[3][1])) * WorldToMetersScale - BaseOffset);
 	OutPosition = BaseOrientation.Inverse().RotateVector(Position);
 
 	OutOrientation = BaseOrientation.Inverse() * OutOrientation;
@@ -1208,7 +1208,7 @@ void FSteamVRHMD::ResetPosition()
 {
 	const FTrackingFrame& TrackingFrame = GetTrackingFrame();
 	FMatrix Pose = ToFMatrix(TrackingFrame.RawPoses[vr::k_unTrackedDeviceIndex_Hmd]);
-	BaseOffset = FVector(-Pose.M[3][2], Pose.M[3][0], Pose.M[3][1]);
+	BaseOffset = FVector(-Pose.M[3][2], Pose.M[3][0], Pose.M[3][1]) * TrackingFrame.WorldToMetersScale;
 }
 
 void FSteamVRHMD::SetBaseRotation(const FRotator& BaseRot)
@@ -1230,6 +1230,16 @@ FQuat FSteamVRHMD::GetBaseOrientation() const
 	return BaseOrientation;
 }
 
+void FSteamVRHMD::SetBasePosition(const FVector& BasePosition)
+{
+	BaseOffset = BasePosition;
+}
+
+FVector FSteamVRHMD::GetBasePosition() const
+{
+	return BaseOffset;
+}
+
 bool FSteamVRHMD::IsStereoEnabled() const
 {
 	return VRSystem && bStereoEnabled && bHmdEnabled;
diff --git a/Engine/Plugins/Runtime/Steam/SteamVR/Source/SteamVR/Private/SteamVRHMD.h b/Engine/Plugins/Runtime/Steam/SteamVR/Source/SteamVR/Private/SteamVRHMD.h
--- a/Engine/Plugins/Runtime/Steam/SteamVR/Source/SteamVR/Private/SteamVRHMD.h
+++ b/Engine/Plugins/Runtime/Steam/SteamVR/Source/SteamVR/Private/SteamVRHMD.h
@@ -154,7 +154,9 @@ public:
 	virtual FRotator GetBaseRotation() const override;
 	virtual void SetBaseOrientation(const FQuat& BaseOrient) override;
 	virtual FQuat GetBaseOrientation() const override;
-
+	virtual void SetBasePosition(const FVector& BasePosition) override;
+	virtual FVector GetBasePosition() const override;
+	
 	virtual void OnEndPlay(FWorldContext& InWorldContext) override;
 	virtual void RecordAnalytics() override;
 
diff --git a/Engine/Source/Runtime/HeadMountedDisplay/Private/HeadMountedDisplayFunctionLibrary.cpp b/Engine/Source/Runtime/HeadMountedDisplay/Private/HeadMountedDisplayFunctionLibrary.cpp
--- a/Engine/Source/Runtime/HeadMountedDisplay/Private/HeadMountedDisplayFunctionLibrary.cpp
+++ b/Engine/Source/Runtime/HeadMountedDisplay/Private/HeadMountedDisplayFunctionLibrary.cpp
@@ -250,6 +250,24 @@ FTransform UHeadMountedDisplayFunctionLibrary::GetTrackingToWorldTransform(UObje
 	return FTransform::Identity;
 }
 
+void UHeadMountedDisplayFunctionLibrary::CalibrateExternalTrackingToHMD(const FTransform& ExternalTrackingTransform)
+{
+	IXRTrackingSystem* TrackingSys = GEngine->XRSystem.Get();
+	if (TrackingSys)
+	{
+		TrackingSys->CalibrateExternalTrackingSource(ExternalTrackingTransform);
+	}
+}
+
+void UHeadMountedDisplayFunctionLibrary::UpdateExternalTrackingHMDPosition(const FTransform& ExternalTrackingTransform)
+{
+	IXRTrackingSystem* TrackingSys = GEngine->XRSystem.Get();
+	if (TrackingSys)
+	{
+		TrackingSys->UpdateExternalTrackingPosition(ExternalTrackingTransform);
+	}
+}
+
 void UHeadMountedDisplayFunctionLibrary::GetVRFocusState(bool& bUseFocus, bool& bHasFocus)
 {
 	IHeadMountedDisplay* HMD = GEngine->XRSystem.IsValid() ? GEngine->XRSystem->GetHMDDevice() : nullptr;
diff --git a/Engine/Source/Runtime/HeadMountedDisplay/Private/XRTrackingSystemBase.cpp b/Engine/Source/Runtime/HeadMountedDisplay/Private/XRTrackingSystemBase.cpp
--- a/Engine/Source/Runtime/HeadMountedDisplay/Private/XRTrackingSystemBase.cpp
+++ b/Engine/Source/Runtime/HeadMountedDisplay/Private/XRTrackingSystemBase.cpp
@@ -74,6 +74,35 @@ void FXRTrackingSystemBase::UpdateTrackingToWorldTransform(const FTransform& Tra
 	CachedTrackingToWorld = TrackingToWorldOverride;
 }
 
+void FXRTrackingSystemBase::CalibrateExternalTrackingSource(const FTransform& ExternalTrackingTransform)
+{
+	FVector HMDPos;
+	FQuat HMDOrientation;
+	GetCurrentPose(IXRTrackingSystem::HMDDeviceId, HMDOrientation, HMDPos);
+
+	FTransform CameraLocalTransform(HMDOrientation, HMDPos);
+
+	CalibratedOffset = ExternalTrackingTransform.GetRelativeTransform(CameraLocalTransform * GetTrackingToWorldTransform());
+	UE_LOG(LogHMD, Log, TEXT("XRTracker is now using calibrated offset to external tracking source: [%s]"), *CalibratedOffset.GetTranslation().ToCompactString());
+}
+
+void FXRTrackingSystemBase::UpdateExternalTrackingPosition(const FTransform& ExternalTrackingTransform)
+{
+	FVector HMDPos;
+	FQuat HMDOrientation;
+	GetCurrentPose(IXRTrackingSystem::HMDDeviceId, HMDOrientation, HMDPos);
+
+	const FTransform CameraLocalTransform(HMDOrientation, HMDPos);
+	const FTransform NewRelativeOffset = ExternalTrackingTransform.GetRelativeTransform(CameraLocalTransform * GetTrackingToWorldTransform());
+	const FTransform DeltaOffset = NewRelativeOffset.GetRelativeTransform(CalibratedOffset);
+
+	//UE_LOG(LogHMD, Warning, TEXT("---> Delta from Calibrated: [%s] (%s)"), *DeltaOffset.GetTranslation().ToCompactString(), *DeltaOffset.GetRotation().Rotator().ToCompactString());
+
+	SetBaseOrientation((DeltaOffset.GetRotation().Inverse() * GetBaseOrientation()).GetNormalized());
+	SetBasePosition(GetBasePosition() - DeltaOffset.GetTranslation());
+}
+
+
 FTransform FXRTrackingSystemBase::RefreshTrackingToWorldTransform(FWorldContext& WorldContext)
 {
 	CachedTrackingToWorld = ComputeTrackingToWorldTransform(WorldContext);
diff --git a/Engine/Source/Runtime/HeadMountedDisplay/Public/HeadMountedDisplayFunctionLibrary.h b/Engine/Source/Runtime/HeadMountedDisplay/Public/HeadMountedDisplayFunctionLibrary.h
--- a/Engine/Source/Runtime/HeadMountedDisplay/Public/HeadMountedDisplayFunctionLibrary.h
+++ b/Engine/Source/Runtime/HeadMountedDisplay/Public/HeadMountedDisplayFunctionLibrary.h
@@ -196,6 +196,29 @@ class HEADMOUNTEDDISPLAY_API UHeadMountedDisplayFunctionLibrary : public UBluepr
 	UFUNCTION(BlueprintCallable, Category = "Input|HeadMountedDisplay",  meta=(WorldContext="WorldContext"))
 	static FTransform GetTrackingToWorldTransform(UObject* WorldContext);
 
+	/**
+	* Called to calibrate the offset transform between an external tracking source and the internal tracking source
+	* (e.g. mocap tracker to and HMD tracker).  This should be called once per session, or when the physical relationship
+	* between the external tracker and internal tracker changes (e.g. it was bumped or reattached).  After calibration,
+	* calling UpdateExternalTrackingPosition will try to correct the internal tracker to the calibrated offset to prevent
+	* drift between the two systems
+	*
+	* @param ExternalTrackingTransform		The transform in world-coordinates, of the reference marker of the external tracking system
+	*/
+	UFUNCTION(BlueprintCallable, Category = "Input|HeadMountedDisplay|ExternalTracking")
+	static void CalibrateExternalTrackingToHMD(const FTransform& ExternalTrackingTransform);
+
+	/**
+	* Called after calibration to attempt to pull the internal tracker (e.g. HMD tracking) in line with the external tracker
+	* (e.g. mocap tracker).  This will set the internal tracker's base offset and rotation to match and realign the two systems.
+	* This can be called every tick, or whenever realignment is desired.  Note that this may cause choppy movement if the two
+	* systems diverge relative to each other, or a big jump if called infrequently when there has been significant drift
+	*
+	* @param ExternalTrackingTransform		The transform in world-coordinates, of the reference marker of the external tracking system
+	*/
+	UFUNCTION(BlueprintCallable, Category = "Input|HeadMountedDisplay|ExternalTracking")
+	static void UpdateExternalTrackingHMDPosition(const FTransform& ExternalTrackingTransform);
+
 	/**
 	 * Returns current state of VR focus.
 	 *
diff --git a/Engine/Source/Runtime/HeadMountedDisplay/Public/IXRTrackingSystem.h b/Engine/Source/Runtime/HeadMountedDisplay/Public/IXRTrackingSystem.h
--- a/Engine/Source/Runtime/HeadMountedDisplay/Public/IXRTrackingSystem.h
+++ b/Engine/Source/Runtime/HeadMountedDisplay/Public/IXRTrackingSystem.h
@@ -232,6 +232,39 @@ public:
 	*/
 	virtual FQuat GetBaseOrientation() const { return FQuat::Identity; }
 
+	/**
+	* Sets base position of the HMD.
+	*
+	* @param BasePosition		(in) the desired offset to be treated as a base position.
+	*/
+	virtual void SetBasePosition(const FVector& BasePosition) {};
+
+	/**
+	* Returns current base position of HMD.
+	*/
+	virtual FVector GetBasePosition() const { return FVector::ZeroVector; }
+
+	/**
+	* Called to calibrate the offset transform between an external tracking source and the internal tracking source
+	* (e.g. mocap tracker to and HMD tracker).  This should be called once per session, or when the physical relationship
+	* between the external tracker and internal tracker changes (e.g. it was bumped or reattached).  After calibration,
+	* calling UpdateExternalTrackingPosition will try to correct the internal tracker to the calibrated offset to prevent
+	* drift between the two systems
+	*
+	* @param ExternalTrackingTransform		(in) The transform in world-coordinates, of the reference marker of the external tracking system
+	*/
+	virtual void CalibrateExternalTrackingSource(const FTransform& ExternalTrackingTransform) {}
+
+	/**
+	* Called after calibration to attempt to pull the internal tracker (e.g. HMD tracking) in line with the external tracker
+	* (e.g. mocap tracker).  This will set the internal tracker's base offset and rotation to match and realign the two systems.
+	* This can be called every tick, or whenever realignment is desired.  Note that this may cause choppy movement if the two
+	* systems diverge relative to each other, or a big jump if called infrequently when there has been significant drift
+	*
+	* @param ExternalTrackingTransform		(in) The transform in world-coordinates, of the reference marker of the external tracking system
+	*/
+	virtual void UpdateExternalTrackingPosition(const FTransform& ExternalTrackingTransform) {}
+
 	/** 
 	 * Get the IXCamera instance for the given device.
 	 *
diff --git a/Engine/Source/Runtime/HeadMountedDisplay/Public/XRTrackingSystemBase.h b/Engine/Source/Runtime/HeadMountedDisplay/Public/XRTrackingSystemBase.h
--- a/Engine/Source/Runtime/HeadMountedDisplay/Public/XRTrackingSystemBase.h
+++ b/Engine/Source/Runtime/HeadMountedDisplay/Public/XRTrackingSystemBase.h
@@ -117,6 +117,27 @@ public:
 	 */
 	virtual void UpdateTrackingToWorldTransform(const FTransform& TrackingToWorldOverride) override;
 
+	/**
+	 * Called to calibrate the offset transform between an external tracking source and the internal tracking source
+	 * (e.g. mocap tracker to and HMD tracker).  This should be called once per session, or when the physical relationship
+	 * between the external tracker and internal tracker changes (e.g. it was bumped or reattached).  After calibration,
+	 * calling UpdateExternalTrackingPosition will try to correct the internal tracker to the calibrated offset to prevent
+	 * drift between the two systems
+	 *
+	 * @param ExternalTrackingTransform		The transform in world-coordinates, of the reference marker of the external tracking system
+	 */
+	virtual void CalibrateExternalTrackingSource(const FTransform& ExternalTrackingTransform) override;
+
+	/**
+	* Called after calibration to attempt to pull the internal tracker (e.g. HMD tracking) in line with the external tracker
+	* (e.g. mocap tracker).  This will set the internal tracker's base offset and rotation to match and realign the two systems.
+	* This can be called every tick, or whenever realignment is desired.  Note that this may cause choppy movement if the two
+	* systems diverge relative to each other, or a big jump if called infrequently when there has been significant drift
+	*
+	* @param ExternalTrackingTransform		The transform in world-coordinates, of the reference marker of the external tracking system
+	*/
+	virtual void UpdateExternalTrackingPosition(const FTransform& ExternalTrackingTransform) override;
+
 	/**
 	 * This method should return the world to meters scale for the current frame.
 	 * Should be callable on both the render and the game threads.
@@ -146,5 +167,13 @@ protected:
 	TSharedPtr< class FDefaultXRCamera, ESPMode::ThreadSafe > XRCamera;
 
 	FTransform CachedTrackingToWorld;
+
+	/**
+	 * If the tracker is trying to lock itself to an external tracking source for drift control,
+	 * this stores the calibrated offset between the external tracking system and the internal tracking system
+	 * (e.g. a location from a mocap system tracker to the HMD's internal IMU).  UpdateExternalTrackingPostion
+	 * will attempt to normalize the internal tracking system to match this calibration when called.
+	 */
+	FTransform CalibratedOffset;
 };
 
