Merged in absolute scale/location/rotation PR #4760 from Wang Hao

diff --git a/Engine/Source/Runtime/HeadMountedDisplay/Private/XRTrackingSystemBase.cpp b/Engine/Source/Runtime/HeadMountedDisplay/Private/XRTrackingSystemBase.cpp
--- a/Engine/Source/Runtime/HeadMountedDisplay/Private/XRTrackingSystemBase.cpp
+++ b/Engine/Source/Runtime/HeadMountedDisplay/Private/XRTrackingSystemBase.cpp
@@ -167,6 +167,21 @@ FTransform FXRTrackingSystemBase::ComputeTrackingToWorldTransform(FWorldContext&
 						if (ViewParent)
 						{
 							TrackingToWorld = ViewParent->GetComponentTransform();
+
+							if (PlayerCamera->bAbsoluteLocation)
+							{
+								TrackingToWorld.SetLocation(FVector::ZeroVector);
+							}
+
+							if (PlayerCamera->bAbsoluteRotation)
+							{
+								TrackingToWorld.SetRotation(FQuat::Identity);
+							}
+
+							if (PlayerCamera->bAbsoluteScale)
+							{
+								TrackingToWorld.SetScale3D(FVector::OneVector);
+							}
 						}
 						// else, if the camera is the root component (not attached to an origin point)
 						// then it is directly relative to the world - the tracking origin is the world's origin (i.e. the identity)
