Fixes HeadMountedDisplayTypes.cpp compilation failure.

diff --git a/Engine/Source/Runtime/HeadMountedDisplay/Private/HeadMountedDisplayTypes.cpp b/Engine/Source/Runtime/HeadMountedDisplay/Private/HeadMountedDisplayTypes.cpp
--- a/Engine/Source/Runtime/HeadMountedDisplay/Private/HeadMountedDisplayTypes.cpp
+++ b/Engine/Source/Runtime/HeadMountedDisplay/Private/HeadMountedDisplayTypes.cpp
@@ -2,6 +2,7 @@
 
 #include "HeadMountedDisplayTypes.h"
 #include "RendererInterface.h"
+#include "DummyRenderResources.h"
 
 DEFINE_LOG_CATEGORY(LogHMD);
 DEFINE_LOG_CATEGORY(LogLoadingSplash);
