Fix for Race Condition in LateUpdateManager. Storing another pointer into the same struct does not work if the entire struct has been freed in the meantime.
(Copy of CL#4997141 from Dev-VR)

diff --git a/Engine/Source/Runtime/HeadMountedDisplay/Private/LateUpdateManager.cpp b/Engine/Source/Runtime/HeadMountedDisplay/Private/LateUpdateManager.cpp
--- a/Engine/Source/Runtime/HeadMountedDisplay/Private/LateUpdateManager.cpp
+++ b/Engine/Source/Runtime/HeadMountedDisplay/Private/LateUpdateManager.cpp
@@ -48,15 +48,40 @@ void FLateUpdateManager::Apply_RenderThread(FSceneInterface* Scene, const FTrans
 	const FTransform NewCameraTransform = NewRelativeTransform * LateUpdateParentToWorld[LateUpdateRenderReadIndex];
 	const FMatrix LateUpdateTransform = (OldCameraTransform.Inverse() * NewCameraTransform).ToMatrixWithScale();
 
-	// Apply delta to the affected scene proxies
-	for (auto PrimitiveInfo : LateUpdatePrimitives[LateUpdateRenderReadIndex])
+	bool bIndicesHaveChanged = false;
+
+	// Apply delta to the cached scene proxies
+	// Also check whether any primitive indices have changed, in case the scene has been modified in the meantime.
+	for (auto PrimitivePair : LateUpdatePrimitives[LateUpdateRenderReadIndex])
 	{
-		FPrimitiveSceneInfo* RetrievedSceneInfo = Scene->GetPrimitiveSceneInfo(*PrimitiveInfo.IndexAddress);
-		FPrimitiveSceneInfo* CachedSceneInfo = PrimitiveInfo.SceneInfo;
-		// If the retrieved scene info is different than our cached scene info then the primitive was removed from the scene
-		if (CachedSceneInfo == RetrievedSceneInfo && CachedSceneInfo->Proxy)
+		FPrimitiveSceneInfo* RetrievedSceneInfo = Scene->GetPrimitiveSceneInfo(PrimitivePair.Value);
+		FPrimitiveSceneInfo* CachedSceneInfo = PrimitivePair.Key;
+
+		// If the retrieved scene info is different than our cached scene info then the scene has changed in the meantime
+		// and we need to search through the entire scene to make sure it still exists.
+		if (CachedSceneInfo != RetrievedSceneInfo)
+		{
+			bIndicesHaveChanged = true;
+			break; // No need to continue here, as we are going to brute force the scene primitives below anyway.
+		}
+		else if (CachedSceneInfo->Proxy)
 		{
 			CachedSceneInfo->Proxy->ApplyLateUpdateTransform(LateUpdateTransform);
+			PrimitivePair.Value = -1; // Set the cached index to -1 to indicate that this primitive was already processed
+		}
+	}
+
+	// Indices have changed, so we need to scan the entire scene for primitives that might still exist
+	if (bIndicesHaveChanged)
+	{
+		int32 Index = 0;
+		FPrimitiveSceneInfo* RetrievedSceneInfo;
+		while(RetrievedSceneInfo = Scene->GetPrimitiveSceneInfo(Index++))
+		{
+			if (RetrievedSceneInfo->Proxy && LateUpdatePrimitives[LateUpdateRenderReadIndex].Contains(RetrievedSceneInfo) && LateUpdatePrimitives[LateUpdateRenderReadIndex][RetrievedSceneInfo] >= 0)
+			{
+				RetrievedSceneInfo->Proxy->ApplyLateUpdateTransform(LateUpdateTransform);
+			}
 		}
 	}
 }
@@ -77,10 +102,7 @@ void FLateUpdateManager::CacheSceneInfo(USceneComponent* Component)
 		FPrimitiveSceneInfo* PrimitiveSceneInfo = PrimitiveComponent->SceneProxy->GetPrimitiveSceneInfo();
 		if (PrimitiveSceneInfo)
 		{
-			LateUpdatePrimitiveInfo PrimitiveInfo;
-			PrimitiveInfo.IndexAddress = PrimitiveSceneInfo->GetIndexAddress();
-			PrimitiveInfo.SceneInfo = PrimitiveSceneInfo;
-			LateUpdatePrimitives[LateUpdateGameWriteIndex].Add(PrimitiveInfo);
+			LateUpdatePrimitives[LateUpdateGameWriteIndex].Emplace(PrimitiveSceneInfo, PrimitiveSceneInfo->GetIndex());
 		}
 	}
 }
diff --git a/Engine/Source/Runtime/HeadMountedDisplay/Public/LateUpdateManager.h b/Engine/Source/Runtime/HeadMountedDisplay/Public/LateUpdateManager.h
--- a/Engine/Source/Runtime/HeadMountedDisplay/Public/LateUpdateManager.h
+++ b/Engine/Source/Runtime/HeadMountedDisplay/Public/LateUpdateManager.h
@@ -30,19 +30,6 @@ public:
 
 private:
 
-	/*
-	 *  Late update primitive info for accessing valid scene proxy info. From the time the info is gathered
-	 *  to the time it is later accessed the render proxy can be deleted. To ensure we only access a proxy that is
-	 *  still valid we cache the primitive's scene info AND a pointer to its own cached index. If the primitive
-	 *  is deleted or removed from the scene then attempting to access it via its index will result in a different
-	 *  scene info than the cached scene info.
-	 */
-	struct LateUpdatePrimitiveInfo
-	{
-		const int32*			IndexAddress;
-		FPrimitiveSceneInfo*	SceneInfo;
-	};
-
 	/** A utility method that calls CacheSceneInfo on ParentComponent and all of its descendants */
 	void GatherLateUpdatePrimitives(USceneComponent* ParentComponent);
 	/** Generates a LateUpdatePrimitiveInfo for the given component if it has a SceneProxy and appends it to the current LateUpdatePrimitives array */
@@ -51,7 +38,7 @@ private:
 	/** Parent world transform used to reconstruct new world transforms for late update scene proxies */
 	FTransform LateUpdateParentToWorld[2];
 	/** Primitives that need late update before rendering */
-	TArray<LateUpdatePrimitiveInfo> LateUpdatePrimitives[2];
+	TMap<FPrimitiveSceneInfo*,int32> LateUpdatePrimitives[2];
 	/** Late Update Info Stale, if this is found true do not late update */
 	bool SkipLateUpdate[2];
 
