Fix for assignment within conditional expression

Copy of CL4998580 in Dev-VR by #fyi Joe.Conley

diff --git a/Engine/Source/Runtime/HeadMountedDisplay/Private/LateUpdateManager.cpp b/Engine/Source/Runtime/HeadMountedDisplay/Private/LateUpdateManager.cpp
--- a/Engine/Source/Runtime/HeadMountedDisplay/Private/LateUpdateManager.cpp
+++ b/Engine/Source/Runtime/HeadMountedDisplay/Private/LateUpdateManager.cpp
@@ -76,12 +76,14 @@ void FLateUpdateManager::Apply_RenderThread(FSceneInterface* Scene, const FTrans
 	{
 		int32 Index = 0;
 		FPrimitiveSceneInfo* RetrievedSceneInfo;
-		while(RetrievedSceneInfo = Scene->GetPrimitiveSceneInfo(Index++))
+		RetrievedSceneInfo = Scene->GetPrimitiveSceneInfo(Index++);
+		while(RetrievedSceneInfo)
 		{
 			if (RetrievedSceneInfo->Proxy && LateUpdatePrimitives[LateUpdateRenderReadIndex].Contains(RetrievedSceneInfo) && LateUpdatePrimitives[LateUpdateRenderReadIndex][RetrievedSceneInfo] >= 0)
 			{
 				RetrievedSceneInfo->Proxy->ApplyLateUpdateTransform(LateUpdateTransform);
 			}
+			RetrievedSceneInfo = Scene->GetPrimitiveSceneInfo(Index++);
 		}
 	}
 }
