Selective move of Dev-VR changes to base XRSwapChain class (and one client usage) to fix build break with new OpenXR plugin code in 4.23. Also fix initializer-list ordering warning in OpenXR plugin.

diff --git a/Engine/Plugins/Runtime/Oculus/OculusVR/Source/OculusHMD/Private/OculusHMD_CustomPresent.cpp b/Engine/Plugins/Runtime/Oculus/OculusVR/Source/OculusHMD/Private/OculusHMD_CustomPresent.cpp
--- a/Engine/Plugins/Runtime/Oculus/OculusVR/Source/OculusHMD/Private/OculusHMD_CustomPresent.cpp
+++ b/Engine/Plugins/Runtime/Oculus/OculusVR/Source/OculusHMD/Private/OculusHMD_CustomPresent.cpp
@@ -293,7 +293,7 @@ FXRSwapChainPtr FCustomPresent::CreateSwapChain_RenderThread(uint32 InSizeX, uin
 		}
 	}
 
-	return MakeShareable(new FXRSwapChain(RHITexture, MoveTemp(RHITextureSwapChain)));
+	return CreateXRSwapChain(MoveTemp(RHITextureSwapChain), RHITexture);
 }
 
 
diff --git a/Engine/Source/Runtime/HeadMountedDisplay/Private/XRSwapChain.cpp b/Engine/Source/Runtime/HeadMountedDisplay/Private/XRSwapChain.cpp
--- a/Engine/Source/Runtime/HeadMountedDisplay/Private/XRSwapChain.cpp
+++ b/Engine/Source/Runtime/HeadMountedDisplay/Private/XRSwapChain.cpp
@@ -60,11 +60,15 @@ FORCEINLINE void CheckInRHIThread()
 #endif
 }
 
-FXRSwapChain::FXRSwapChain(FRHITexture* InRHITexture, TArray<FTextureRHIRef>&& InRHITextureSwapChain)
-	: RHITexture(InRHITexture)
+FXRSwapChain::FXRSwapChain(TArray<FTextureRHIRef>&& InRHITextureSwapChain, const FTextureRHIRef& AliasedTexture)
+	: RHITexture(AliasedTexture)
 	, RHITextureSwapChain(InRHITextureSwapChain)
 	, SwapChainIndex_RHIThread(0)
 {
+	check(RHITexture);
+	// @todo: The *correct* way to create the RHITexture object (that's just an alias over swap-chain textures)
+	//		  would be via a new RHI API to create an aliased texture. For now we just let the clients of this class
+	//		  create themselves (since they all create differently).
 	RHITexture->SetName(TEXT("XRSwapChainAliasedTexture"));
 	for (int ChainElement = 0; ChainElement < RHITextureSwapChain.Num(); ++ChainElement)
 	{
@@ -108,7 +112,7 @@ void FXRSwapChain::GenerateMips_RenderThread(FRHICommandListImmediate& RHICmdLis
 }
 
 
-void FXRSwapChain::IncrementSwapChainIndex_RHIThread()
+void FXRSwapChain::IncrementSwapChainIndex_RHIThread(int64 /* TimeoutNanoseconds */)
 {
 	CheckInRHIThread();
 
diff --git a/Engine/Source/Runtime/HeadMountedDisplay/Public/XRSwapChain.h b/Engine/Source/Runtime/HeadMountedDisplay/Public/XRSwapChain.h
--- a/Engine/Source/Runtime/HeadMountedDisplay/Public/XRSwapChain.h
+++ b/Engine/Source/Runtime/HeadMountedDisplay/Public/XRSwapChain.h
@@ -11,9 +11,10 @@
 class HEADMOUNTEDDISPLAY_API FXRSwapChain : public TSharedFromThis<FXRSwapChain, ESPMode::ThreadSafe>
 {
 public:
-	FXRSwapChain(FRHITexture* InRHITexture, TArray<FTextureRHIRef>&& InRHITextureSwapChain);
+	FXRSwapChain(TArray<FTextureRHIRef>&& InRHITextureSwapChain, const FTextureRHIRef& AliasedTexture);
 	virtual ~FXRSwapChain();
 
+	const FTextureRHIRef& GetTextureRef() const { return RHITexture; }
 	FRHITexture* GetTexture() const { return RHITexture.GetReference(); }
 	FRHITexture2D* GetTexture2D() const { return RHITexture->GetTexture2D(); }
 	FRHITexture2DArray* GetTexture2DArray() const { return RHITexture->GetTexture2DArray(); }
@@ -23,8 +24,8 @@ public:
 	void GenerateMips_RenderThread(FRHICommandListImmediate& RHICmdList);
 	uint32 GetSwapChainIndex_RHIThread() { return SwapChainIndex_RHIThread; }
 
-	virtual void IncrementSwapChainIndex_RHIThread();
-	virtual void ReleaseCurrentImage_RHIThread() {}		
+	virtual void IncrementSwapChainIndex_RHIThread(int64 TimeoutNanoseconds = 0); // Default to no timeout (immediate).
+	virtual void ReleaseCurrentImage_RHIThread() {}
 
 protected:
 	virtual void ReleaseResources_RHIThread();
@@ -36,3 +37,12 @@ protected:
 
 typedef TSharedPtr<FXRSwapChain, ESPMode::ThreadSafe> FXRSwapChainPtr;
 
+template<typename T = FXRSwapChain, typename... Types>
+FXRSwapChainPtr CreateXRSwapChain(TArray<FTextureRHIRef>&& InRHITextureSwapChain, const FTextureRHIRef& InRHIAliasedTexture, Types... InExtraParameters)
+{
+	check(InRHITextureSwapChain.Num() >= 1);
+	return MakeShareable(new T(MoveTemp(InRHITextureSwapChain), InRHIAliasedTexture, InExtraParameters...));
+}
+
+
+
