Fix collection-modified-during-ranged-for assertion in late-update class. Potential race-conditions between game thread and rendering thread could cause the game thread to modify the late update list while the rendering thread was consuming data. The fix is to wrap all collection accesses and index modifications with a critical section/scoped lock.

diff --git a/Engine/Source/Runtime/HeadMountedDisplay/Private/DefaultXRCamera.cpp b/Engine/Source/Runtime/HeadMountedDisplay/Private/DefaultXRCamera.cpp
--- a/Engine/Source/Runtime/HeadMountedDisplay/Private/DefaultXRCamera.cpp
+++ b/Engine/Source/Runtime/HeadMountedDisplay/Private/DefaultXRCamera.cpp
@@ -183,11 +183,6 @@ void FDefaultXRCamera::PreRenderViewFamily_RenderThread(FRHICommandListImmediate
 	}
 }
 
-void FDefaultXRCamera::PostRenderViewFamily_RenderThread(FRHICommandListImmediate& RHICmdList, FSceneViewFamily& InViewFamily)
-{
-	LateUpdate.PostRender_RenderThread();
-}
-
 void FDefaultXRCamera::SetupViewFamily(FSceneViewFamily& InViewFamily)
 {
 	static const auto CVarAllowMotionBlurInVR = IConsoleManager::Get().FindTConsoleVariableDataInt(TEXT("vr.AllowMotionBlurInVR"));
diff --git a/Engine/Source/Runtime/HeadMountedDisplay/Private/LateUpdateManager.cpp b/Engine/Source/Runtime/HeadMountedDisplay/Private/LateUpdateManager.cpp
--- a/Engine/Source/Runtime/HeadMountedDisplay/Private/LateUpdateManager.cpp
+++ b/Engine/Source/Runtime/HeadMountedDisplay/Private/LateUpdateManager.cpp
@@ -4,55 +4,55 @@
 #include "PrimitiveSceneProxy.h"
 #include "Components/PrimitiveComponent.h"
 #include "PrimitiveSceneInfo.h"
+#include "HeadMountedDisplayTypes.h"
 
 FLateUpdateManager::FLateUpdateManager()
-	: LateUpdateGameWriteIndex(0)
-	, LateUpdateRenderReadIndex(0)
 {
-	SkipLateUpdate[0] = false;
-	SkipLateUpdate[1] = false;
 }
 
 void FLateUpdateManager::Setup(const FTransform& ParentToWorld, USceneComponent* Component, bool bSkipLateUpdate)
 {
 	check(IsInGameThread());
 
-	LateUpdateParentToWorld[LateUpdateGameWriteIndex] = ParentToWorld;
-	LateUpdatePrimitives[LateUpdateGameWriteIndex].Reset();
-	GatherLateUpdatePrimitives(Component);
-	SkipLateUpdate[LateUpdateGameWriteIndex] = bSkipLateUpdate;
-
-	LateUpdateGameWriteIndex = (LateUpdateGameWriteIndex + 1) % 2;
-}
+	FScopeLock Lock(&StateCriticalSection);
 
-bool FLateUpdateManager::GetSkipLateUpdate_RenderThread() const
-{
-	return SkipLateUpdate[LateUpdateRenderReadIndex];
+	GameThreadState.ParentToWorld = ParentToWorld;
+	GatherLateUpdatePrimitives(Component);
+	GameThreadState.bSkip = bSkipLateUpdate;
+	++GameThreadState.TrackingNumber;
 }
 
 void FLateUpdateManager::Apply_RenderThread(FSceneInterface* Scene, const FTransform& OldRelativeTransform, const FTransform& NewRelativeTransform)
 {
 	check(IsInRenderingThread());
 
-	if (!LateUpdatePrimitives[LateUpdateRenderReadIndex].Num())
 	{
-		return;
+		// Only grab the lock to protect access to the GameThreadState collection - the RenderThreadState is
+		// only ever accessed from the rendering thread.
+		FScopeLock Lock(&StateCriticalSection);
+		RenderThreadState.bSkip = GameThreadState.bSkip;
+		RenderThreadState.ParentToWorld = GameThreadState.ParentToWorld;
+		RenderThreadState.Primitives = MoveTemp(GameThreadState.Primitives);
+
+		GameThreadState.Primitives.Reset();
+
+		RenderThreadState.TrackingNumber = GameThreadState.TrackingNumber;
 	}
 
-	if (GetSkipLateUpdate_RenderThread())
+	if (!RenderThreadState.Primitives.Num() || RenderThreadState.bSkip)
 	{
 		return;
 	}
 
-	const FTransform OldCameraTransform = OldRelativeTransform * LateUpdateParentToWorld[LateUpdateRenderReadIndex];
-	const FTransform NewCameraTransform = NewRelativeTransform * LateUpdateParentToWorld[LateUpdateRenderReadIndex];
+	const FTransform OldCameraTransform = OldRelativeTransform * RenderThreadState.ParentToWorld;
+	const FTransform NewCameraTransform = NewRelativeTransform * RenderThreadState.ParentToWorld;
 	const FMatrix LateUpdateTransform = (OldCameraTransform.Inverse() * NewCameraTransform).ToMatrixWithScale();
 
 	bool bIndicesHaveChanged = false;
 
 	// Apply delta to the cached scene proxies
 	// Also check whether any primitive indices have changed, in case the scene has been modified in the meantime.
-	for (auto PrimitivePair : LateUpdatePrimitives[LateUpdateRenderReadIndex])
+	for (auto PrimitivePair : RenderThreadState.Primitives)
 	{
 		FPrimitiveSceneInfo* RetrievedSceneInfo = Scene->GetPrimitiveSceneInfo(PrimitivePair.Value);
 		FPrimitiveSceneInfo* CachedSceneInfo = PrimitivePair.Key;
@@ -79,7 +79,7 @@ void FLateUpdateManager::Apply_RenderThread(FSceneInterface* Scene, const FTrans
 		RetrievedSceneInfo = Scene->GetPrimitiveSceneInfo(Index++);
 		while(RetrievedSceneInfo)
 		{
-			if (RetrievedSceneInfo->Proxy && LateUpdatePrimitives[LateUpdateRenderReadIndex].Contains(RetrievedSceneInfo) && LateUpdatePrimitives[LateUpdateRenderReadIndex][RetrievedSceneInfo] >= 0)
+			if (RetrievedSceneInfo->Proxy && RenderThreadState.Primitives.Contains(RetrievedSceneInfo) && RenderThreadState.Primitives[RetrievedSceneInfo] >= 0)
 			{
 				RetrievedSceneInfo->Proxy->ApplyLateUpdateTransform(LateUpdateTransform);
 			}
@@ -88,13 +88,6 @@ void FLateUpdateManager::Apply_RenderThread(FSceneInterface* Scene, const FTrans
 	}
 }
 
-void FLateUpdateManager::PostRender_RenderThread()
-{
-	LateUpdatePrimitives[LateUpdateRenderReadIndex].Reset();
-	SkipLateUpdate[LateUpdateRenderReadIndex] = false;
-	LateUpdateRenderReadIndex = (LateUpdateRenderReadIndex + 1) % 2;
-}
-
 void FLateUpdateManager::CacheSceneInfo(USceneComponent* Component)
 {
 	// If a scene proxy is present, cache it
@@ -104,7 +97,7 @@ void FLateUpdateManager::CacheSceneInfo(USceneComponent* Component)
 		FPrimitiveSceneInfo* PrimitiveSceneInfo = PrimitiveComponent->SceneProxy->GetPrimitiveSceneInfo();
 		if (PrimitiveSceneInfo && PrimitiveSceneInfo->IsIndexValid())
 		{
-			LateUpdatePrimitives[LateUpdateGameWriteIndex].Emplace(PrimitiveSceneInfo, PrimitiveSceneInfo->GetIndex());
+			GameThreadState.Primitives.Emplace(PrimitiveSceneInfo, PrimitiveSceneInfo->GetIndex());
 		}
 	}
 }
diff --git a/Engine/Source/Runtime/HeadMountedDisplay/Private/MotionControllerComponent.cpp b/Engine/Source/Runtime/HeadMountedDisplay/Private/MotionControllerComponent.cpp
--- a/Engine/Source/Runtime/HeadMountedDisplay/Private/MotionControllerComponent.cpp
+++ b/Engine/Source/Runtime/HeadMountedDisplay/Private/MotionControllerComponent.cpp
@@ -582,15 +582,6 @@ void UMotionControllerComponent::FViewExtension::PreRenderViewFamily_RenderThrea
 	LateUpdate.Apply_RenderThread(InViewFamily.Scene, OldTransform, NewTransform);
 }
 
-void UMotionControllerComponent::FViewExtension::PostRenderViewFamily_RenderThread(FRHICommandListImmediate& RHICmdList, FSceneViewFamily& InViewFamily)
-{
-	if (!MotionControllerComponent)
-	{
-		return;
-	}
-	LateUpdate.PostRender_RenderThread();
-}
-
 bool UMotionControllerComponent::FViewExtension::IsActiveThisFrame(class FViewport* InViewport) const
 {
 	check(IsInGameThread());
diff --git a/Engine/Source/Runtime/HeadMountedDisplay/Private/MotionDelayBuffer.cpp b/Engine/Source/Runtime/HeadMountedDisplay/Private/MotionDelayBuffer.cpp
--- a/Engine/Source/Runtime/HeadMountedDisplay/Private/MotionDelayBuffer.cpp
+++ b/Engine/Source/Runtime/HeadMountedDisplay/Private/MotionDelayBuffer.cpp
@@ -496,11 +496,6 @@ void FMotionDelayClient::PostRenderViewFamily_RenderThread(FRHICommandListImmedi
 {
 	if (MotionDelayService_Impl::PostRenderCleanupId_RenderThread != ViewFamily.FrameNumber)
 	{
-		for (const FTargetTransform& Transform : TargetTransforms_RenderThread)
-		{
-			Transform.DelayTarget->LateUpdate.PostRender_RenderThread();
-		}
-
 		MotionDelayService_Impl::PostRenderCleanupId_RenderThread = ViewFamily.FrameNumber;
 	}
 }
diff --git a/Engine/Source/Runtime/HeadMountedDisplay/Public/DefaultXRCamera.h b/Engine/Source/Runtime/HeadMountedDisplay/Public/DefaultXRCamera.h
--- a/Engine/Source/Runtime/HeadMountedDisplay/Public/DefaultXRCamera.h
+++ b/Engine/Source/Runtime/HeadMountedDisplay/Public/DefaultXRCamera.h
@@ -78,7 +78,6 @@ public:
 	virtual void BeginRenderViewFamily(FSceneViewFamily& InViewFamily) override;
 	virtual void PreRenderView_RenderThread(FRHICommandListImmediate& RHICmdList, FSceneView& InView) override;
 	virtual void PreRenderViewFamily_RenderThread(FRHICommandListImmediate& RHICmdList, FSceneViewFamily& InViewFamily) override;
-	virtual void PostRenderViewFamily_RenderThread(FRHICommandListImmediate& RHICmdList, FSceneViewFamily& InViewFamily) override;
 	virtual bool IsActiveThisFrame(class FViewport* InViewport) const override;
 	
 
diff --git a/Engine/Source/Runtime/HeadMountedDisplay/Public/LateUpdateManager.h b/Engine/Source/Runtime/HeadMountedDisplay/Public/LateUpdateManager.h
--- a/Engine/Source/Runtime/HeadMountedDisplay/Public/LateUpdateManager.h
+++ b/Engine/Source/Runtime/HeadMountedDisplay/Public/LateUpdateManager.h
@@ -20,14 +20,11 @@ public:
 	void Setup(const FTransform& ParentToWorld, USceneComponent* Component, bool bSkipLateUpdate);
 
 	/** Returns true if the LateUpdateSetup data is stale. */
-	bool GetSkipLateUpdate_RenderThread() const;
+	bool GetSkipLateUpdate_RenderThread() const { return RenderThreadState.bSkip; }
 
 	/** Apply the late update delta to the cached components */
 	void Apply_RenderThread(FSceneInterface* Scene, const FTransform& OldRelativeTransform, const FTransform& NewRelativeTransform);
 
-	/** Increments the double buffered read index, etc. - in prep for the next render frame (read: MUST be called for each frame Setup() was called on). */
-	void PostRender_RenderThread();
-
 private:
 
 	/** A utility method that calls CacheSceneInfo on ParentComponent and all of its descendants */
@@ -35,15 +32,26 @@ private:
 	/** Generates a LateUpdatePrimitiveInfo for the given component if it has a SceneProxy and appends it to the current LateUpdatePrimitives array */
 	void CacheSceneInfo(USceneComponent* Component);
 
-	/** Parent world transform used to reconstruct new world transforms for late update scene proxies */
-	FTransform LateUpdateParentToWorld[2];
-	/** Primitives that need late update before rendering */
-	TMap<FPrimitiveSceneInfo*,int32> LateUpdatePrimitives[2];
-	/** Late Update Info Stale, if this is found true do not late update */
-	bool SkipLateUpdate[2];
-
-	int32 LateUpdateGameWriteIndex;
-	int32 LateUpdateRenderReadIndex;
-
+	struct FLateUpdateState
+	{
+		FLateUpdateState()
+			: ParentToWorld(FTransform::Identity)
+			, bSkip(false)
+			, TrackingNumber(-1)
+		{}
+
+		/** Parent world transform used to reconstruct new world transforms for late update scene proxies */
+		FTransform ParentToWorld;
+		/** Primitives that need late update before rendering */
+		TMap<FPrimitiveSceneInfo*, int32> Primitives;
+		/** Late Update Info Stale, if this is found true do not late update */
+		bool bSkip;
+		/** Frame tracking number - used to flag if the game and render threads get badly out of sync */
+		int64 TrackingNumber;
+	};
+
+	FLateUpdateState GameThreadState;
+	FLateUpdateState RenderThreadState;
+	FCriticalSection StateCriticalSection;
 };
 
diff --git a/Engine/Source/Runtime/HeadMountedDisplay/Public/MotionControllerComponent.h b/Engine/Source/Runtime/HeadMountedDisplay/Public/MotionControllerComponent.h
--- a/Engine/Source/Runtime/HeadMountedDisplay/Public/MotionControllerComponent.h
+++ b/Engine/Source/Runtime/HeadMountedDisplay/Public/MotionControllerComponent.h
@@ -153,7 +153,6 @@ private:
 		virtual void BeginRenderViewFamily(FSceneViewFamily& InViewFamily) override;
 		virtual void PreRenderView_RenderThread(FRHICommandListImmediate& RHICmdList, FSceneView& InView) override {}
 		virtual void PreRenderViewFamily_RenderThread(FRHICommandListImmediate& RHICmdList, FSceneViewFamily& InViewFamily) override;
-		virtual void PostRenderViewFamily_RenderThread(FRHICommandListImmediate& RHICmdList, FSceneViewFamily& InViewFamily) override;
 		virtual int32 GetPriority() const override { return -10; }
 		virtual bool IsActiveThisFrame(class FViewport* InViewport) const;
 
