UMotionControllerComponent::SendRenderTransform_Concurrent is runing parallel with render thread, but RenderThreadRelativeTransform is accessing by lateupdate on rendering thread, which results in a race condition when rendering threads start late or when the main thread is running fast.

Also this is logically wrong, since, just like other transforms, RenderThreadRelativeTransform update should be kicked off by SendRenderTransform_Concurrent, but queued up for next frame rendering before lateupdate accesses it. ENQEUE_RENDER_COMMAND solves the problem well.

diff --git a/Engine/Source/Runtime/HeadMountedDisplay/Private/MotionControllerComponent.cpp b/Engine/Source/Runtime/HeadMountedDisplay/Private/MotionControllerComponent.cpp
--- a/Engine/Source/Runtime/HeadMountedDisplay/Private/MotionControllerComponent.cpp
+++ b/Engine/Source/Runtime/HeadMountedDisplay/Private/MotionControllerComponent.cpp
@@ -101,8 +101,23 @@ void UMotionControllerComponent::CreateRenderState_Concurrent()
 
 void UMotionControllerComponent::SendRenderTransform_Concurrent()
 {
-	RenderThreadRelativeTransform = GetRelativeTransform();
-	RenderThreadComponentScale = GetComponentScale();
+	struct FPrimitiveUpdateRenderThreadRelativeTransformParams
+	{
+		FTransform RenderThreadRelativeTransform;
+		FVector RenderThreadComponentScale;
+	};
+
+	FPrimitiveUpdateRenderThreadRelativeTransformParams UpdateParams;
+	UpdateParams.RenderThreadRelativeTransform = GetRelativeTransform();
+	UpdateParams.RenderThreadComponentScale = GetComponentScale();
+
+	ENQUEUE_RENDER_COMMAND(UpdateRTRelativeTransformCommand)(
+		[UpdateParams, this](FRHICommandListImmediate& RHICmdList)
+	{
+		RenderThreadRelativeTransform = UpdateParams.RenderThreadRelativeTransform;
+		RenderThreadComponentScale = UpdateParams.RenderThreadComponentScale;
+	});
+
 	Super::SendRenderTransform_Concurrent();
 }
 
