Get rid of GWorlds
#rnx

diff --git a/Engine/Plugins/Experimental/SimpleHMD/Source/SimpleHMD/Private/SimpleHMD.cpp b/Engine/Plugins/Experimental/SimpleHMD/Source/SimpleHMD/Private/SimpleHMD.cpp
--- a/Engine/Plugins/Experimental/SimpleHMD/Source/SimpleHMD/Private/SimpleHMD.cpp
+++ b/Engine/Plugins/Experimental/SimpleHMD/Source/SimpleHMD/Private/SimpleHMD.cpp
@@ -287,9 +287,9 @@ void FSimpleHMD::SetupViewFamily(FSceneViewFamily& InViewFamily)
 	InViewFamily.EngineShowFlags.HMDDistortion = true;
 	InViewFamily.EngineShowFlags.StereoRendering = IsStereoEnabled();
 
-	if (GWorld)
+	if (UWorld* World = GWorld)
 	{
-		WorldToMeters = GWorld->GetWorldSettings()->WorldToMeters;
+		WorldToMeters = World->GetWorldSettings()->WorldToMeters;
 	}
 }
 
diff --git a/Engine/Source/Editor/UnrealEd/Private/Commandlets/ContentCommandlets.cpp b/Engine/Source/Editor/UnrealEd/Private/Commandlets/ContentCommandlets.cpp
--- a/Engine/Source/Editor/UnrealEd/Private/Commandlets/ContentCommandlets.cpp
+++ b/Engine/Source/Editor/UnrealEd/Private/Commandlets/ContentCommandlets.cpp
@@ -1550,6 +1550,7 @@ void UResavePackagesCommandlet::PerformAdditionalOperations(class UWorld* World,
 
 			if (bShouldBuildHLOD)
 			{
+				UWorld* World = GWorld;
 				UE_LOG( LogContentCommandlet, Display, TEXT( "Generating HLOD data for %s" ), *World->GetOutermost()->GetName() );
 
 				if( !ForceHLODSetupAsset.IsEmpty() )
@@ -1557,7 +1558,7 @@ void UResavePackagesCommandlet::PerformAdditionalOperations(class UWorld* World,
 					TSubclassOf<UHierarchicalLODSetup> NewHLODSetupAsset = LoadClass<UHierarchicalLODSetup>( NULL, *ForceHLODSetupAsset, NULL, LOAD_None, NULL );
 					if( NewHLODSetupAsset != nullptr )
 					{
-						GWorld->GetWorldSettings()->HLODSetupAsset = NewHLODSetupAsset;
+						World->GetWorldSettings()->HLODSetupAsset = NewHLODSetupAsset;
 					}
 					else
 					{
@@ -1568,16 +1569,16 @@ void UResavePackagesCommandlet::PerformAdditionalOperations(class UWorld* World,
 				// Force HLOD support on for this level if we were asked to
 				if( bForceEnableHLODForLevel )
 				{
-					GWorld->GetWorldSettings()->bEnableHierarchicalLODSystem = true;
+					World->GetWorldSettings()->bEnableHierarchicalLODSystem = true;
 				}
 
 				// Use a single cluster for all actors in the level if we were asked to
 				if( bForceSingleClusterForLevel )
 				{
-					GWorld->GetWorldSettings()->bGenerateSingleClusterForLevel = true;
+					World->GetWorldSettings()->bGenerateSingleClusterForLevel = true;
 				}
 
-				FHierarchicalLODBuilder Builder(GWorld);
+				FHierarchicalLODBuilder Builder(World);
 
 				if (bForceClusterGeneration)
 				{
@@ -1613,7 +1614,7 @@ void UResavePackagesCommandlet::PerformAdditionalOperations(class UWorld* World,
 
 				// Get the list of packages needs to be saved.
 				TSet<UPackage*> PackagesToSave;
-				for(ULevel* Level : GWorld->GetLevels())
+				for(ULevel* Level : World->GetLevels())
 				{
 					if(Level->bIsVisible)
 					{
diff --git a/Engine/Source/Editor/UnrealEd/Private/Commandlets/PackageUtilities.cpp b/Engine/Source/Editor/UnrealEd/Private/Commandlets/PackageUtilities.cpp
--- a/Engine/Source/Editor/UnrealEd/Private/Commandlets/PackageUtilities.cpp
+++ b/Engine/Source/Editor/UnrealEd/Private/Commandlets/PackageUtilities.cpp
@@ -2722,9 +2722,9 @@ int32 UReplaceActorCommandlet::Main(const FString& Params)
 	}
 
 	// UEditorEngine::FinishDestroy() expects GWorld to exist
-	if( GWorld )
+	if( UWorld* World = GWorld )
 	{
-		GWorld->DestroyWorld( false );
+		World->DestroyWorld( false );
 	}
 	GWorld = UWorld::CreateWorld(EWorldType::Editor, false );
 	return 0;
diff --git a/Engine/Source/Editor/UnrealEd/Private/EditorBuildUtils.cpp b/Engine/Source/Editor/UnrealEd/Private/EditorBuildUtils.cpp
--- a/Engine/Source/Editor/UnrealEd/Private/EditorBuildUtils.cpp
+++ b/Engine/Source/Editor/UnrealEd/Private/EditorBuildUtils.cpp
@@ -711,12 +711,13 @@ bool FEditorBuildUtils::PrepForAutomatedBuild( const FEditorAutomatedBuildSettin
 	if ( bBuildSuccessful )
 	{
 		bool bVisibilityToggled = false;
-		if ( !FLevelUtils::IsLevelVisible( GWorld->PersistentLevel ) )
+		UWorld* World = GWorld;
+		if ( !FLevelUtils::IsLevelVisible( World->PersistentLevel ) )
 		{
-			EditorLevelUtils::SetLevelVisibility( GWorld->PersistentLevel, true, false );
+			EditorLevelUtils::SetLevelVisibility( World->PersistentLevel, true, false );
 			bVisibilityToggled = true;
 		}
-		for (ULevelStreaming* CurStreamingLevel : GWorld->GetStreamingLevels())
+		for (ULevelStreaming* CurStreamingLevel : World->GetStreamingLevels())
 		{
 			if ( CurStreamingLevel && !FLevelUtils::IsStreamingLevelVisibleInEditor( CurStreamingLevel ) )
 			{
@@ -726,7 +727,7 @@ bool FEditorBuildUtils::PrepForAutomatedBuild( const FEditorAutomatedBuildSettin
 		}
 		if ( bVisibilityToggled )
 		{
-			GWorld->FlushLevelStreaming();
+			World->FlushLevelStreaming();
 		}
 	}
 
diff --git a/Engine/Source/Editor/UnrealEd/Private/SEditorViewportViewMenu.cpp b/Engine/Source/Editor/UnrealEd/Private/SEditorViewportViewMenu.cpp
--- a/Engine/Source/Editor/UnrealEd/Private/SEditorViewportViewMenu.cpp
+++ b/Engine/Source/Editor/UnrealEd/Private/SEditorViewportViewMenu.cpp
@@ -317,7 +317,8 @@ TSharedRef<SWidget> SEditorViewportViewMenu::GenerateViewMenuContent() const
 					{
 						const FEditorViewportCommands& BaseViewportCommands = FEditorViewportCommands::Get();
 
-						const ERHIFeatureLevel::Type FeatureLevel = (IsInGameThread() && GWorld) ? (ERHIFeatureLevel::Type)GWorld->FeatureLevel : GMaxRHIFeatureLevel;
+						UWorld* World = GWorld;
+						const ERHIFeatureLevel::Type FeatureLevel = (IsInGameThread() && World) ? (ERHIFeatureLevel::Type)World->FeatureLevel : GMaxRHIFeatureLevel;
 
 						Menu.BeginSection("OptimizationViewmodes", LOCTEXT("OptimizationSubMenuHeader", "Optimization Viewmodes"));
 						{
diff --git a/Engine/Source/Editor/UnrealEd/Private/StaticLightingSystem/StaticLightingSystem.cpp b/Engine/Source/Editor/UnrealEd/Private/StaticLightingSystem/StaticLightingSystem.cpp
--- a/Engine/Source/Editor/UnrealEd/Private/StaticLightingSystem/StaticLightingSystem.cpp
+++ b/Engine/Source/Editor/UnrealEd/Private/StaticLightingSystem/StaticLightingSystem.cpp
@@ -288,17 +288,18 @@ void FStaticLightingManager::CreateStaticLightingSystem(const FLightingBuildOpti
 
 		bBuildReflectionCapturesOnFinish = !Options.bOnlyBuildVisibility;
 
-		for (ULevel* Level : GWorld->GetLevels())
+		UWorld* World = GWorld;
+		for (ULevel* Level : World->GetLevels())
 		{
 			if (Level->bIsLightingScenario && Level->bIsVisible)
 			{
-				StaticLightingSystems.Emplace(new FStaticLightingSystem(Options, GWorld, Level));
+				StaticLightingSystems.Emplace(new FStaticLightingSystem(Options, World, Level));
 			}
 		}
 
 		if (StaticLightingSystems.Num() == 0)
 		{
-			StaticLightingSystems.Emplace(new FStaticLightingSystem(Options, GWorld, nullptr));
+			StaticLightingSystems.Emplace(new FStaticLightingSystem(Options, World, nullptr));
 		}
 
 		ActiveStaticLightingSystem = StaticLightingSystems[0].Get();
@@ -314,7 +315,7 @@ void FStaticLightingManager::CreateStaticLightingSystem(const FLightingBuildOpti
 			// BeginLightmassProcess returns false if there are errors or no precomputed lighting is allowed. Handle both cases.
 			static const auto AllowStaticLightingVar = IConsoleManager::Get().FindTConsoleVariableDataInt(TEXT("r.AllowStaticLighting"));
 			const bool bAllowStaticLighting = (!AllowStaticLightingVar || AllowStaticLightingVar->GetValueOnGameThread() != 0);
-			const bool bForceNoPrecomputedLighting = GWorld->GetWorldSettings()->bForceNoPrecomputedLighting || !bAllowStaticLighting;
+			const bool bForceNoPrecomputedLighting = World->GetWorldSettings()->bForceNoPrecomputedLighting || !bAllowStaticLighting;
 
 
 			if (bForceNoPrecomputedLighting)
diff --git a/Engine/Source/Runtime/HeadMountedDisplay/Private/HeadMountedDisplayConsoleCommands.cpp b/Engine/Source/Runtime/HeadMountedDisplay/Private/HeadMountedDisplayConsoleCommands.cpp
--- a/Engine/Source/Runtime/HeadMountedDisplay/Private/HeadMountedDisplayConsoleCommands.cpp
+++ b/Engine/Source/Runtime/HeadMountedDisplay/Private/HeadMountedDisplayConsoleCommands.cpp
@@ -57,9 +57,10 @@ static void DrawDebugTrackingSensorLocations(UCanvas* Canvas, APlayerController*
 		return;
 	}
 
+	UWorld* World = (PlayerController ? PlayerController->GetWorld() : World);
 	if (!PlayerController)
 	{
-		PlayerController = GWorld->GetFirstPlayerController();
+		PlayerController = World->GetFirstPlayerController();
 		if (!PlayerController)
 		{
 			return;
@@ -133,30 +134,30 @@ static void DrawDebugTrackingSensorLocations(UCanvas* Canvas, APlayerController*
 		FVector NearBR = Matrix.TransformPosition(EdgeBR * SensorProperties.NearPlane);
 
 		// Draw a point at the sensor position
-		DrawDebugPoint(GWorld, Tip, 5, FrustrumColor);
+		DrawDebugPoint(World, Tip, 5, FrustrumColor);
 
 		// Draw the four edges of the pyramid
-		DrawDebugLine(GWorld, Tip, BaseTR, FrustrumColor);
-		DrawDebugLine(GWorld, Tip, BaseTL, FrustrumColor);
-		DrawDebugLine(GWorld, Tip, BaseBL, FrustrumColor);
-		DrawDebugLine(GWorld, Tip, BaseBR, FrustrumColor);
+		DrawDebugLine(World, Tip, BaseTR, FrustrumColor);
+		DrawDebugLine(World, Tip, BaseTL, FrustrumColor);
+		DrawDebugLine(World, Tip, BaseBL, FrustrumColor);
+		DrawDebugLine(World, Tip, BaseBR, FrustrumColor);
 
 		// Draw the base (far plane)
-		DrawDebugLine(GWorld, BaseTR, BaseTL, FrustrumColor);
-		DrawDebugLine(GWorld, BaseTL, BaseBL, FrustrumColor);
-		DrawDebugLine(GWorld, BaseBL, BaseBR, FrustrumColor);
-		DrawDebugLine(GWorld, BaseBR, BaseTR, FrustrumColor);
+		DrawDebugLine(World, BaseTR, BaseTL, FrustrumColor);
+		DrawDebugLine(World, BaseTL, BaseBL, FrustrumColor);
+		DrawDebugLine(World, BaseBL, BaseBR, FrustrumColor);
+		DrawDebugLine(World, BaseBR, BaseTR, FrustrumColor);
 
 		// Draw the near plane
-		DrawDebugLine(GWorld, NearTR, NearTL, FrustrumColor);
-		DrawDebugLine(GWorld, NearTL, NearBL, FrustrumColor);
-		DrawDebugLine(GWorld, NearBL, NearBR, FrustrumColor);
-		DrawDebugLine(GWorld, NearBR, NearTR, FrustrumColor);
+		DrawDebugLine(World, NearTR, NearTL, FrustrumColor);
+		DrawDebugLine(World, NearTL, NearBL, FrustrumColor);
+		DrawDebugLine(World, NearBL, NearBR, FrustrumColor);
+		DrawDebugLine(World, NearBR, NearTR, FrustrumColor);
 
 		// Draw a center line from the sensor to the focal point
 		FVector CenterLine = Matrix.TransformPosition(FVector(SensorProperties.CameraDistance, 0, 0));
-		DrawDebugLine(GWorld, Tip, CenterLine, CenterLineColor);
-		DrawDebugPoint(GWorld, CenterLine, 5, CenterLineColor);
+		DrawDebugLine(World, Tip, CenterLine, CenterLineColor);
+		DrawDebugPoint(World, CenterLine, 5, CenterLineColor);
 	}
 }
 
