Fix typo
#rnx

diff --git a/Engine/Source/Runtime/HeadMountedDisplay/Private/HeadMountedDisplayConsoleCommands.cpp b/Engine/Source/Runtime/HeadMountedDisplay/Private/HeadMountedDisplayConsoleCommands.cpp
--- a/Engine/Source/Runtime/HeadMountedDisplay/Private/HeadMountedDisplayConsoleCommands.cpp
+++ b/Engine/Source/Runtime/HeadMountedDisplay/Private/HeadMountedDisplayConsoleCommands.cpp
@@ -57,7 +57,7 @@ static void DrawDebugTrackingSensorLocations(UCanvas* Canvas, APlayerController*
 		return;
 	}
 
-	UWorld* World = (PlayerController ? PlayerController->GetWorld() : World);
+	UWorld* World = (PlayerController ? PlayerController->GetWorld() : GWorld);
 	if (!PlayerController)
 	{
 		PlayerController = World->GetFirstPlayerController();
