UEVR-1268 Composite debug canvas layer to spectator screen
-The debug canvas layer is now copied onto the spectator screen if both of those features are active.
-The layer is mapped onto the spectator screen with letterboxing, so its aspect ratio will be preserved but additional unused screen space may exist (usually on the left and right).  The layer resolution is unlikely to match screen resolution so the display is expected to be somewhat fuzzier than normal.
-Spectator modes that use api level support rather than the SpectatorScreenController will continue to function as before.

diff --git a/Engine/Plugins/Runtime/Oculus/OculusVR/Source/OculusHMD/Private/OculusHMD.cpp b/Engine/Plugins/Runtime/Oculus/OculusVR/Source/OculusHMD/Private/OculusHMD.cpp
--- a/Engine/Plugins/Runtime/Oculus/OculusVR/Source/OculusHMD/Private/OculusHMD.cpp
+++ b/Engine/Plugins/Runtime/Oculus/OculusVR/Source/OculusHMD/Private/OculusHMD.cpp
@@ -1331,7 +1331,7 @@ namespace OculusHMD
 	}
 
 
-	void FOculusHMD::CopyTexture_RenderThread(FRHICommandListImmediate& RHICmdList, FTexture2DRHIParamRef SrcTexture, FIntRect SrcRect, FTexture2DRHIParamRef DstTexture, FIntRect DstRect, bool bClearBlack) const
+	void FOculusHMD::CopyTexture_RenderThread(FRHICommandListImmediate& RHICmdList, FTexture2DRHIParamRef SrcTexture, FIntRect SrcRect, FTexture2DRHIParamRef DstTexture, FIntRect DstRect, bool bClearBlack, bool bNoAlpha) const
 	{
 		if (bClearBlack)
 		{
@@ -1342,7 +1342,7 @@ namespace OculusHMD
 		}
 
 		check(CustomPresent);
-		CustomPresent->CopyTexture_RenderThread(RHICmdList, DstTexture, SrcTexture, DstRect, SrcRect);
+		CustomPresent->CopyTexture_RenderThread(RHICmdList, DstTexture, SrcTexture, DstRect, SrcRect, false, bNoAlpha);
 	}
 
 
diff --git a/Engine/Plugins/Runtime/Oculus/OculusVR/Source/OculusHMD/Private/OculusHMD.h b/Engine/Plugins/Runtime/Oculus/OculusVR/Source/OculusHMD/Private/OculusHMD.h
--- a/Engine/Plugins/Runtime/Oculus/OculusVR/Source/OculusHMD/Private/OculusHMD.h
+++ b/Engine/Plugins/Runtime/Oculus/OculusVR/Source/OculusHMD/Private/OculusHMD.h
@@ -112,6 +112,7 @@ public:
 	//virtual FVector GetAudioListenerOffset() const override;
 	virtual bool GetHMDDistortionEnabled(EShadingPath ShadingPath) const override;
 	//virtual void BeginRendering_RenderThread(const FTransform& NewRelativeTransform, FRHICommandListImmediate& RHICmdList, FSceneViewFamily& ViewFamily) override;
+	//virtual bool IsSpectatorScreenActive() const override;
 	//virtual class ISpectatorScreenController* GetSpectatorScreenController() override;
 	//virtual class ISpectatorScreenController const* GetSpectatorScreenController() const override;
 	//virtual float GetDistortionScalingFactor() const override;
@@ -154,7 +155,6 @@ public:
 	virtual FMatrix GetStereoProjectionMatrix(const enum EStereoscopicPass StereoPassType) const override;
 	virtual void InitCanvasFromView(class FSceneView* InView, class UCanvas* Canvas) override;
 	//virtual void GetEyeRenderParams_RenderThread(const struct FRenderingCompositePassContext& Context, FVector2D& EyeToSrcUVScaleValue, FVector2D& EyeToSrcUVOffsetValue) const override;
-	//virtual bool IsSpectatorScreenActive() const override;
 	virtual void RenderTexture_RenderThread(class FRHICommandListImmediate& RHICmdList, class FRHITexture2D* BackBuffer, class FRHITexture2D* SrcTexture, FVector2D WindowSize) const override;
 	//virtual void SetClippingPlanes(float NCP, float FCP) override;
 	virtual IStereoRenderTargetManager* GetRenderTargetManager() override { return this; }
@@ -165,7 +165,7 @@ public:
 	// FHeadMoundedDisplayBase interface
 	virtual FVector2D GetEyeCenterPoint_RenderThread(EStereoscopicPass StereoPassType) const override;
 	virtual FIntRect GetFullFlatEyeRect_RenderThread(FTexture2DRHIRef EyeTexture) const override;
-	virtual void CopyTexture_RenderThread(FRHICommandListImmediate& RHICmdList, FTexture2DRHIParamRef SrcTexture, FIntRect SrcRect, FTexture2DRHIParamRef DstTexture, FIntRect DstRect, bool bClearBlack) const override;
+	virtual void CopyTexture_RenderThread(FRHICommandListImmediate& RHICmdList, FTexture2DRHIParamRef SrcTexture, FIntRect SrcRect, FTexture2DRHIParamRef DstTexture, FIntRect DstRect, bool bClearBlack, bool bNoAlpha) const override;
 	virtual bool PopulateAnalyticsAttributes(TArray<struct FAnalyticsEventAttribute>& EventAttributes) override;
 
 	// FXRRenderTargetManager interface
diff --git a/Engine/Plugins/Runtime/Oculus/OculusVR/Source/OculusHMD/Private/OculusHMD_CustomPresent.cpp b/Engine/Plugins/Runtime/Oculus/OculusVR/Source/OculusHMD/Private/OculusHMD_CustomPresent.cpp
--- a/Engine/Plugins/Runtime/Oculus/OculusVR/Source/OculusHMD/Private/OculusHMD_CustomPresent.cpp
+++ b/Engine/Plugins/Runtime/Oculus/OculusVR/Source/OculusHMD/Private/OculusHMD_CustomPresent.cpp
@@ -372,8 +372,8 @@ void FCustomPresent::CopyTexture_RenderThread(FRHICommandListImmediate& RHICmdLi
 		}
 		else
 		{
-			// for mirror window
-			GraphicsPSOInit.BlendState = TStaticBlendState<>::GetRHI();
+			// for mirror window, write RGBA, RGB = src.rgb * src.a + dst.rgb * (1 - src.a), A = src.a * 1 + dst.a * (1 - src a)
+			GraphicsPSOInit.BlendState = TStaticBlendState<CW_RGBA, BO_Add, BF_SourceAlpha, BF_InverseSourceAlpha, BO_Add, BF_One, BF_InverseSourceAlpha>::GetRHI();
 		}
 	}
 
diff --git a/Engine/Plugins/Runtime/Oculus/OculusVR/Source/OculusHMD/Private/OculusHMD_SpectatorScreenController.cpp b/Engine/Plugins/Runtime/Oculus/OculusVR/Source/OculusHMD/Private/OculusHMD_SpectatorScreenController.cpp
--- a/Engine/Plugins/Runtime/Oculus/OculusVR/Source/OculusHMD/Private/OculusHMD_SpectatorScreenController.cpp
+++ b/Engine/Plugins/Runtime/Oculus/OculusVR/Source/OculusHMD/Private/OculusHMD_SpectatorScreenController.cpp
@@ -18,7 +18,7 @@ FOculusHMD_SpectatorScreenController::FOculusHMD_SpectatorScreenController(FOcul
 {
 }
 
-void FOculusHMD_SpectatorScreenController::RenderSpectatorScreen_RenderThread(FRHICommandListImmediate& RHICmdList, FRHITexture2D* BackBuffer, FTexture2DRHIRef RenderTexture, FVector2D WindowSize) const
+void FOculusHMD_SpectatorScreenController::RenderSpectatorScreen_RenderThread(FRHICommandListImmediate& RHICmdList, FRHITexture2D* BackBuffer, FTexture2DRHIRef RenderTexture, FVector2D WindowSize)
 {
 	if (OculusHMD->GetCustomPresent_Internal())
 	{
@@ -33,7 +33,7 @@ void FOculusHMD_SpectatorScreenController::RenderSpectatorModeUndistorted(FRHICo
 	FIntRect DestRect(0, 0, TargetTexture->GetSizeX() / 2, TargetTexture->GetSizeY());
 	for (int i = 0; i < 2; ++i)
 	{
-		OculusHMD->CopyTexture_RenderThread(RHICmdList, EyeTexture, Settings->EyeRenderViewport[i], TargetTexture, DestRect, false);
+		OculusHMD->CopyTexture_RenderThread(RHICmdList, EyeTexture, Settings->EyeRenderViewport[i], TargetTexture, DestRect, false, true);
 		DestRect.Min.X += TargetTexture->GetSizeX() / 2;
 		DestRect.Max.X += TargetTexture->GetSizeX() / 2;
 	}
@@ -48,7 +48,7 @@ void FOculusHMD_SpectatorScreenController::RenderSpectatorModeDistorted(FRHIComm
 	{
 		FIntRect SrcRect(0, 0, MirrorTexture->GetSizeX(), MirrorTexture->GetSizeY());
 		FIntRect DestRect(0, 0, TargetTexture->GetSizeX(), TargetTexture->GetSizeY());
-		OculusHMD->CopyTexture_RenderThread(RHICmdList, MirrorTexture, SrcRect, TargetTexture, DestRect, false);
+		OculusHMD->CopyTexture_RenderThread(RHICmdList, MirrorTexture, SrcRect, TargetTexture, DestRect, false, true);
 	}
 }
 
@@ -59,7 +59,7 @@ void FOculusHMD_SpectatorScreenController::RenderSpectatorModeSingleEye(FRHIComm
 	const FIntRect SrcRect= Settings->EyeRenderViewport[0];
 	const FIntRect DstRect(0, 0, TargetTexture->GetSizeX(), TargetTexture->GetSizeY());
 
-	OculusHMD->CopyTexture_RenderThread(RHICmdList, EyeTexture, SrcRect, TargetTexture, DstRect, false);
+	OculusHMD->CopyTexture_RenderThread(RHICmdList, EyeTexture, SrcRect, TargetTexture, DstRect, false, true);
 }
 
 } // namespace OculusHMD
diff --git a/Engine/Plugins/Runtime/Oculus/OculusVR/Source/OculusHMD/Private/OculusHMD_SpectatorScreenController.h b/Engine/Plugins/Runtime/Oculus/OculusVR/Source/OculusHMD/Private/OculusHMD_SpectatorScreenController.h
--- a/Engine/Plugins/Runtime/Oculus/OculusVR/Source/OculusHMD/Private/OculusHMD_SpectatorScreenController.h
+++ b/Engine/Plugins/Runtime/Oculus/OculusVR/Source/OculusHMD/Private/OculusHMD_SpectatorScreenController.h
@@ -18,7 +18,7 @@ class FOculusHMD_SpectatorScreenController : public FDefaultSpectatorScreenContr
 public:
 	FOculusHMD_SpectatorScreenController(class FOculusHMD* InOculusHMD);
 
-	virtual void RenderSpectatorScreen_RenderThread(FRHICommandListImmediate& RHICmdList, FRHITexture2D* BackBuffer, FTexture2DRHIRef RenderTarget, FVector2D WindowSize) const override;
+	virtual void RenderSpectatorScreen_RenderThread(FRHICommandListImmediate& RHICmdList, FRHITexture2D* BackBuffer, FTexture2DRHIRef RenderTarget, FVector2D WindowSize) override;
 	virtual void RenderSpectatorModeUndistorted(FRHICommandListImmediate& RHICmdList, FTexture2DRHIRef TargetTexture, FTexture2DRHIRef EyeTexture, FTexture2DRHIRef OtherTexture, FVector2D WindowSize) override;
 	virtual void RenderSpectatorModeDistorted(FRHICommandListImmediate& RHICmdList, FTexture2DRHIRef TargetTexture, FTexture2DRHIRef EyeTexture, FTexture2DRHIRef OtherTexture, FVector2D WindowSize) override;
 	virtual void RenderSpectatorModeSingleEye(FRHICommandListImmediate& RHICmdList, FTexture2DRHIRef TargetTexture, FTexture2DRHIRef EyeTexture, FTexture2DRHIRef OtherTexture, FVector2D WindowSize) override;
diff --git a/Engine/Plugins/Runtime/Steam/SteamVR/Source/SteamVR/Private/SteamVRHMD.h b/Engine/Plugins/Runtime/Steam/SteamVR/Source/SteamVR/Private/SteamVRHMD.h
--- a/Engine/Plugins/Runtime/Steam/SteamVR/Source/SteamVR/Private/SteamVRHMD.h
+++ b/Engine/Plugins/Runtime/Steam/SteamVR/Source/SteamVR/Private/SteamVRHMD.h
@@ -59,6 +59,8 @@ struct FSteamVRLayer
 	{}
 
 	// Required by TStereoLayerManager:
+	void SetLayerId(uint32 InId) { LayerDesc.SetLayerId(InId); }
+	uint32 GetLayerId() const { return LayerDesc.GetLayerId(); }
 	friend bool GetLayerDescMember(const FSteamVRLayer& Layer, FLayerDesc& OutLayerDesc);
 	friend void SetLayerDescMember(FSteamVRLayer& Layer, const FLayerDesc& InLayerDesc);
 	friend void MarkLayerTextureForUpdate(FSteamVRLayer& Layer);
@@ -225,6 +227,7 @@ public:
 	// IStereoLayers interface
 	// Create/Set/Get/Destroy inherited from TStereoLayerManager
 	virtual void UpdateSplashScreen() override;
+	virtual void GetAllocatedTexture(uint32 LayerId, FTextureRHIRef &Texture, FTextureRHIRef &LeftTexture) override;
 
 	/** IHeadMountedDisplayVulkanExtensions */
 	virtual bool GetVulkanInstanceExtensionsRequired( TArray<const ANSICHAR*>& Out ) override;
@@ -235,7 +238,7 @@ private:
 	void CreateSpectatorScreenController();
 public:
 	virtual FIntRect GetFullFlatEyeRect_RenderThread(FTexture2DRHIRef EyeTexture) const override;
-	virtual void CopyTexture_RenderThread(FRHICommandListImmediate& RHICmdList, FTexture2DRHIParamRef SrcTexture, FIntRect SrcRect, FTexture2DRHIParamRef DstTexture, FIntRect DstRect, bool bClearBlack) const override;
+	virtual void CopyTexture_RenderThread(FRHICommandListImmediate& RHICmdList, FTexture2DRHIParamRef SrcTexture, FIntRect SrcRect, FTexture2DRHIParamRef DstTexture, FIntRect DstRect, bool bClearBlack, bool bNoAlpha) const override;
 
 	class BridgeBaseImpl : public FXRRenderBridge
 	{
diff --git a/Engine/Plugins/Runtime/Steam/SteamVR/Source/SteamVR/Private/SteamVRSpectatorScreen.cpp b/Engine/Plugins/Runtime/Steam/SteamVR/Source/SteamVR/Private/SteamVRSpectatorScreen.cpp
--- a/Engine/Plugins/Runtime/Steam/SteamVR/Source/SteamVR/Private/SteamVRSpectatorScreen.cpp
+++ b/Engine/Plugins/Runtime/Steam/SteamVR/Source/SteamVR/Private/SteamVRSpectatorScreen.cpp
@@ -34,7 +34,7 @@ FIntRect FSteamVRHMD::GetFullFlatEyeRect_RenderThread(FTexture2DRHIRef EyeTextur
 	return FIntRect(EyeTexture->GetSizeX() * SrcNormRectMin.X, EyeTexture->GetSizeY() * SrcNormRectMin.Y, EyeTexture->GetSizeX() * SrcNormRectMax.X, EyeTexture->GetSizeY() * SrcNormRectMax.Y);
 }
 
-void FSteamVRHMD::CopyTexture_RenderThread(FRHICommandListImmediate& RHICmdList, FTexture2DRHIParamRef SrcTexture, FIntRect SrcRect, FTexture2DRHIParamRef DstTexture, FIntRect DstRect, bool bClearBlack) const
+void FSteamVRHMD::CopyTexture_RenderThread(FRHICommandListImmediate& RHICmdList, FTexture2DRHIParamRef SrcTexture, FIntRect SrcRect, FTexture2DRHIParamRef DstTexture, FIntRect DstRect, bool bClearBlack, bool bNoAlpha) const
 {
 	check(IsInRenderingThread());
 
@@ -67,7 +67,7 @@ void FSteamVRHMD::CopyTexture_RenderThread(FRHICommandListImmediate& RHICmdList,
 
 	FGraphicsPipelineStateInitializer GraphicsPSOInit;
 	RHICmdList.ApplyCachedRenderTargets(GraphicsPSOInit);
-	GraphicsPSOInit.BlendState = TStaticBlendState<>::GetRHI();
+	GraphicsPSOInit.BlendState = bNoAlpha ? TStaticBlendState<>::GetRHI() : TStaticBlendState<CW_RGBA, BO_Add, BF_SourceAlpha, BF_InverseSourceAlpha, BO_Add, BF_One, BF_InverseSourceAlpha>::GetRHI();
 	GraphicsPSOInit.RasterizerState = TStaticRasterizerState<>::GetRHI();
 	GraphicsPSOInit.DepthStencilState = TStaticDepthStencilState<false, CF_Always>::GetRHI();
 	GraphicsPSOInit.PrimitiveType = PT_TriangleList;
diff --git a/Engine/Plugins/Runtime/Steam/SteamVR/Source/SteamVR/Private/SteamVRStereoLayers.cpp b/Engine/Plugins/Runtime/Steam/SteamVR/Source/SteamVR/Private/SteamVRStereoLayers.cpp
--- a/Engine/Plugins/Runtime/Steam/SteamVR/Source/SteamVR/Private/SteamVRStereoLayers.cpp
+++ b/Engine/Plugins/Runtime/Steam/SteamVR/Source/SteamVR/Private/SteamVRStereoLayers.cpp
@@ -334,6 +334,49 @@ void FSteamVRHMD::UpdateStereoLayers_RenderThread()
 	}
 }
 
+void FSteamVRHMD::GetAllocatedTexture(uint32 LayerId, FTextureRHIRef &Texture, FTextureRHIRef &LeftTexture)
+{
+	Texture = LeftTexture = nullptr;
+	FSteamVRLayer* LayerFound = nullptr;
+
+	if (IsInRenderingThread())
+	{
+		ForEachLayer([&](uint32 /* unused */, FSteamVRLayer& Layer)
+		{
+			if (Layer.GetLayerId() == LayerId)
+			{
+				LayerFound = &Layer;
+			}
+		});
+	}
+	else
+	{
+		// Only supporting the use of this function on RenderingThread.
+		check(false);
+		return;
+	}
+
+	if (LayerFound && LayerFound->LayerDesc.Texture)
+	{
+		switch (LayerFound->LayerDesc.ShapeType)
+		{
+		case IStereoLayers::CubemapLayer:
+			Texture = LayerFound->LayerDesc.Texture->GetTextureCube();
+			LeftTexture = LayerFound->LayerDesc.LeftTexture ? LayerFound->LayerDesc.LeftTexture->GetTextureCube() : nullptr;
+			break;
+
+		case IStereoLayers::CylinderLayer:
+		case IStereoLayers::QuadLayer:
+			Texture = LayerFound->LayerDesc.Texture->GetTexture2D();
+			LeftTexture = LayerFound->LayerDesc.LeftTexture ? LayerFound->LayerDesc.LeftTexture->GetTexture2D() : nullptr;
+			break;
+
+		default:
+			break;
+		}
+	}
+}
+
 //=============================================================================
 IStereoLayers* FSteamVRHMD::GetStereoLayers ()
 {
diff --git a/Engine/Source/Runtime/Engine/Private/Slate/DebugCanvas.cpp b/Engine/Source/Runtime/Engine/Private/Slate/DebugCanvas.cpp
--- a/Engine/Source/Runtime/Engine/Private/Slate/DebugCanvas.cpp
+++ b/Engine/Source/Runtime/Engine/Private/Slate/DebugCanvas.cpp
@@ -9,6 +9,9 @@
 #include "Framework/Application/SlateApplication.h"
 #include "IStereoLayers.h"
 #include "StereoRendering.h"
+#include "IXRTrackingSystem.h"
+#include "ISpectatorScreenController.h"
+#include "IHeadMountedDisplay.h"
 
 /**
  * Simple representation of the backbuffer that the debug canvas renders to
@@ -228,6 +231,20 @@ void FDebugCanvasDrawer::DrawRenderThread(FRHICommandListImmediate& RHICmdList,
 			if (StereoLayers)
 			{
 				StereoLayers->GetAllocatedTexture(LayerID, HMDSwapchain, HMDNull);
+
+				// If drawing to a layer tell the spectator screen controller to copy that layer to the spectator screen.
+				if (LayerID != INVALID_LAYER_ID && GEngine && GEngine->XRSystem)
+				{
+					IHeadMountedDisplay* HMD = GEngine->XRSystem->GetHMDDevice();
+					if (HMD)
+					{
+						ISpectatorScreenController* SpectatorScreenController = HMD->GetSpectatorScreenController();
+						if (SpectatorScreenController)
+						{
+							SpectatorScreenController->QueueDebugCanvasLayerID(LayerID);
+						}
+					}
+				}
 			}
 			RT = reinterpret_cast<FTexture2DRHIRef*>(HMDSwapchain == nullptr ? &LayerTexture->GetRenderTargetItem().ShaderResourceTexture : &HMDSwapchain);
 		}
diff --git a/Engine/Source/Runtime/Engine/Public/IStereoLayers.h b/Engine/Source/Runtime/Engine/Public/IStereoLayers.h
--- a/Engine/Source/Runtime/Engine/Public/IStereoLayers.h
+++ b/Engine/Source/Runtime/Engine/Public/IStereoLayers.h
@@ -45,6 +45,12 @@ public:
 	 */
 	struct FLayerDesc
 	{
+		void SetLayerId(uint32 InId) { Id = InId; }
+		uint32 GetLayerId() const { return Id; }
+
+		const static uint32	INVALID_LAYER_ID = UINT_MAX;
+		// The layer's ID
+		uint32				Id			= INVALID_LAYER_ID;
 		// View space transform
 		FTransform			Transform	 = FTransform::Identity;
 		// Size of rendered quad
diff --git a/Engine/Source/Runtime/Engine/Public/StereoRendering.h b/Engine/Source/Runtime/Engine/Public/StereoRendering.h
--- a/Engine/Source/Runtime/Engine/Public/StereoRendering.h
+++ b/Engine/Source/Runtime/Engine/Public/StereoRendering.h
@@ -126,9 +126,6 @@ public:
 	 */
 	virtual void InitCanvasFromView(class FSceneView* InView, class UCanvas* Canvas) = 0;
 
-	// Are we outputting so a Spectator Screen now.
-	virtual bool IsSpectatorScreenActive() const { return false; }
-
 	// Renders texture into a backbuffer. Could be empty if no rendertarget texture is used, or if direct-rendering 
 	// through RHI bridge is implemented. 
 	virtual void RenderTexture_RenderThread(class FRHICommandListImmediate& RHICmdList, class FRHITexture2D* BackBuffer, class FRHITexture2D* SrcTexture, FVector2D WindowSize) const {}
diff --git a/Engine/Source/Runtime/HeadMountedDisplay/Private/DefaultSpectatorScreenController.cpp b/Engine/Source/Runtime/HeadMountedDisplay/Private/DefaultSpectatorScreenController.cpp
--- a/Engine/Source/Runtime/HeadMountedDisplay/Private/DefaultSpectatorScreenController.cpp
+++ b/Engine/Source/Runtime/HeadMountedDisplay/Private/DefaultSpectatorScreenController.cpp
@@ -10,6 +10,7 @@
 #include "Engine/Texture.h"
 #include "HeadMountedDisplayBase.h"
 #include "SceneUtils.h" // for SCOPED_DRAW_EVENT()
+#include "IStereoLayers.h"
 
 
 FDefaultSpectatorScreenController::FDefaultSpectatorScreenController(FHeadMountedDisplayBase* InHMDDevice)
@@ -65,12 +66,16 @@ void FDefaultSpectatorScreenController::SetSpectatorScreenModeTexturePlusEyeLayo
 	}
 }
 
+void FDefaultSpectatorScreenController::QueueDebugCanvasLayerID(int32 LayerID)
+{
+	DebugCanvasLayerIDs.Add(LayerID);
+}
+
 FSpectatorScreenRenderDelegate* FDefaultSpectatorScreenController::GetSpectatorScreenRenderDelegate_RenderThread()
 {
 	return &SpectatorScreenDelegate_RenderThread;
 }
 
-
 struct FRHISetSpectatorScreenTexture final : public FRHICommand<FRHISetSpectatorScreenTexture>
 {
 	FDefaultSpectatorScreenController* SpectatorScreenController;
@@ -229,7 +234,7 @@ void FDefaultSpectatorScreenController::UpdateSpectatorScreenMode_RenderThread()
 	}
 }
 
-void FDefaultSpectatorScreenController::RenderSpectatorScreen_RenderThread(FRHICommandListImmediate& RHICmdList, FRHITexture2D* BackBuffer, FTexture2DRHIRef SrcTexture, FVector2D WindowSize) const
+void FDefaultSpectatorScreenController::RenderSpectatorScreen_RenderThread(FRHICommandListImmediate& RHICmdList, FRHITexture2D* BackBuffer, FTexture2DRHIRef SrcTexture, FVector2D WindowSize)
 {
 	SCOPED_NAMED_EVENT_TEXT("RenderSocialScreen_RenderThread()", FColor::Magenta);
 
@@ -240,6 +245,27 @@ void FDefaultSpectatorScreenController::RenderSpectatorScreen_RenderThread(FRHIC
 		SCOPED_DRAW_EVENT(RHICmdList, SpectatorScreen)
 		SpectatorScreenDelegate_RenderThread.Execute(RHICmdList, BackBuffer, SrcTexture, SpectatorScreenTexture_RenderThread, WindowSize);
 	}
+
+	// Apply the debug canvas layer.
+	IStereoLayers* StereoLayers = HMDDevice->GetStereoLayers();
+	if (StereoLayers)
+	{
+		const FIntRect DstRect(0, 0, BackBuffer->GetSizeX(), BackBuffer->GetSizeY());
+		for (int32 LayerID : DebugCanvasLayerIDs)
+		{
+			FTextureRHIRef LayerTexture = nullptr, HMDNull = nullptr;
+			StereoLayers->GetAllocatedTexture(LayerID, LayerTexture, HMDNull);
+			if (LayerTexture)
+			{
+				FTexture2DRHIRef LayerTexture2D = LayerTexture->GetTexture2D();
+				check(LayerTexture2D.IsValid());  // Debug canvas layer should be a 2d layer
+				const FIntRect LayerRect(0, 0, LayerTexture2D->GetSizeX(), LayerTexture2D->GetSizeY());
+				const FIntRect DstRectLetterboxed = Helpers::GetLetterboxedDestRect(LayerRect, DstRect);
+				HMDDevice->CopyTexture_RenderThread(RHICmdList, LayerTexture2D, LayerRect, BackBuffer, DstRectLetterboxed, false, false);
+			}
+		}
+		DebugCanvasLayerIDs.Empty();
+	}
 }
 
 
@@ -253,7 +279,7 @@ void FDefaultSpectatorScreenController::RenderSpectatorModeUndistorted(FRHIComma
 	const FIntRect SrcRect(0, 0, EyeTexture->GetSizeX(), EyeTexture->GetSizeY());
 	const FIntRect DstRect(0, 0, TargetTexture->GetSizeX(), TargetTexture->GetSizeY());
 
-	HMDDevice->CopyTexture_RenderThread(RHICmdList, EyeTexture, SrcRect, TargetTexture, DstRect, false);
+	HMDDevice->CopyTexture_RenderThread(RHICmdList, EyeTexture, SrcRect, TargetTexture, DstRect, false, true);
 }
 
 void FDefaultSpectatorScreenController::RenderSpectatorModeDistorted(FRHICommandListImmediate& RHICmdList, FTexture2DRHIRef TargetTexture, FTexture2DRHIRef EyeTexture, FTexture2DRHIRef OtherTexture, FVector2D WindowSize)
@@ -268,7 +294,7 @@ void FDefaultSpectatorScreenController::RenderSpectatorModeSingleEye(FRHICommand
 	const FIntRect SrcRect(0, 0, EyeTexture->GetSizeX() / 2, EyeTexture->GetSizeY());
 	const FIntRect DstRect(0, 0, TargetTexture->GetSizeX(), TargetTexture->GetSizeY());
 
-	HMDDevice->CopyTexture_RenderThread(RHICmdList, EyeTexture, SrcRect, TargetTexture, DstRect, false);
+	HMDDevice->CopyTexture_RenderThread(RHICmdList, EyeTexture, SrcRect, TargetTexture, DstRect, false, true);
 }
 
 void FDefaultSpectatorScreenController::RenderSpectatorModeSingleEyeLetterboxed(FRHICommandListImmediate& RHICmdList, FTexture2DRHIRef TargetTexture, FTexture2DRHIRef EyeTexture, FTexture2DRHIRef OtherTexture, FVector2D WindowSize)
@@ -277,7 +303,7 @@ void FDefaultSpectatorScreenController::RenderSpectatorModeSingleEyeLetterboxed(
 	const FIntRect DstRect(0, 0, TargetTexture->GetSizeX(), TargetTexture->GetSizeY());
 	const FIntRect DstRectLetterboxed = Helpers::GetLetterboxedDestRect(SrcRect, DstRect);
 
-	HMDDevice->CopyTexture_RenderThread(RHICmdList, EyeTexture, SrcRect, TargetTexture, DstRectLetterboxed, true);
+	HMDDevice->CopyTexture_RenderThread(RHICmdList, EyeTexture, SrcRect, TargetTexture, DstRectLetterboxed, true, true);
 }
 
 void FDefaultSpectatorScreenController::RenderSpectatorModeSingleEyeCroppedToFill(FRHICommandListImmediate& RHICmdList, FTexture2DRHIRef TargetTexture, FTexture2DRHIRef EyeTexture, FTexture2DRHIRef OtherTexture, FVector2D WindowSize)
@@ -288,7 +314,7 @@ void FDefaultSpectatorScreenController::RenderSpectatorModeSingleEyeCroppedToFil
 
 	const FIntRect SrcCroppedToFitRect = Helpers::GetEyeCroppedToFitRect(HMDDevice->GetEyeCenterPoint_RenderThread(EStereoscopicPass::eSSP_LEFT_EYE), SrcRect, WindowRect);
 
-	HMDDevice->CopyTexture_RenderThread(RHICmdList, EyeTexture, SrcCroppedToFitRect, TargetTexture, DstRect, false);
+	HMDDevice->CopyTexture_RenderThread(RHICmdList, EyeTexture, SrcCroppedToFitRect, TargetTexture, DstRect, false, true);
 }
 
 void FDefaultSpectatorScreenController::RenderSpectatorModeTexture(FRHICommandListImmediate& RHICmdList, FTexture2DRHIRef TargetTexture, FTexture2DRHIRef EyeTexture, FTexture2DRHIRef OtherTexture, FVector2D WindowSize)
@@ -302,7 +328,7 @@ void FDefaultSpectatorScreenController::RenderSpectatorModeTexture(FRHICommandLi
 	const FIntRect SrcRect(0, 0, SrcTexture->GetSizeX(), SrcTexture->GetSizeY());
 	const FIntRect DstRect(0, 0, TargetTexture->GetSizeX(), TargetTexture->GetSizeY());
 
-	HMDDevice->CopyTexture_RenderThread(RHICmdList, SrcTexture, SrcRect, TargetTexture, DstRect, false);
+	HMDDevice->CopyTexture_RenderThread(RHICmdList, SrcTexture, SrcRect, TargetTexture, DstRect, false, true);
 }
 
 void FDefaultSpectatorScreenController::RenderSpectatorModeMirrorAndTexture(FRHICommandListImmediate& RHICmdList, FTexture2DRHIRef TargetTexture, FTexture2DRHIRef EyeTexture, FTexture2DRHIRef OtherTexture, FVector2D WindowSize)
@@ -324,13 +350,13 @@ void FDefaultSpectatorScreenController::RenderSpectatorModeMirrorAndTexture(FRHI
 
 	if (SpectatorScreenModeTexturePlusEyeLayout_RenderThread.bDrawEyeFirst)
 	{
-		HMDDevice->CopyTexture_RenderThread(RHICmdList, EyeTexture, CroppedEyeSrcRect, TargetTexture, EyeDstRect, bClearBlack);
-		HMDDevice->CopyTexture_RenderThread(RHICmdList, OtherTextureLocal, OtherSrcRect, TargetTexture, OtherDstRect, false);
+		HMDDevice->CopyTexture_RenderThread(RHICmdList, EyeTexture, CroppedEyeSrcRect, TargetTexture, EyeDstRect, bClearBlack, true);
+		HMDDevice->CopyTexture_RenderThread(RHICmdList, OtherTextureLocal, OtherSrcRect, TargetTexture, OtherDstRect, false, true);
 	}
 	else
 	{
-		HMDDevice->CopyTexture_RenderThread(RHICmdList, OtherTextureLocal, OtherSrcRect, TargetTexture, OtherDstRect, bClearBlack);
-		HMDDevice->CopyTexture_RenderThread(RHICmdList, EyeTexture, CroppedEyeSrcRect, TargetTexture, EyeDstRect, false);
+		HMDDevice->CopyTexture_RenderThread(RHICmdList, OtherTextureLocal, OtherSrcRect, TargetTexture, OtherDstRect, bClearBlack, true);
+		HMDDevice->CopyTexture_RenderThread(RHICmdList, EyeTexture, CroppedEyeSrcRect, TargetTexture, EyeDstRect, false, true);
 	}
 }
 
diff --git a/Engine/Source/Runtime/HeadMountedDisplay/Private/DefaultStereoLayers.cpp b/Engine/Source/Runtime/HeadMountedDisplay/Private/DefaultStereoLayers.cpp
--- a/Engine/Source/Runtime/HeadMountedDisplay/Private/DefaultStereoLayers.cpp
+++ b/Engine/Source/Runtime/HeadMountedDisplay/Private/DefaultStereoLayers.cpp
@@ -313,3 +313,46 @@ void FDefaultStereoLayers::SetupViewFamily(FSceneViewFamily& InViewFamily)
 	HmdTransform = FTransform(HmdOrientation, HmdPosition);
 }
 
+void FDefaultStereoLayers::GetAllocatedTexture(uint32 LayerId, FTextureRHIRef &Texture, FTextureRHIRef &LeftTexture)
+{
+	Texture = LeftTexture = nullptr;
+	FLayerDesc* LayerFound = nullptr;
+
+	if (IsInRenderingThread())
+	{
+		for (int32 LayerIndex = 0; LayerIndex < RenderThreadLayers.Num(); LayerIndex++)
+		{
+			if (RenderThreadLayers[LayerIndex].GetLayerId() == LayerId)
+			{
+				LayerFound = &RenderThreadLayers[LayerIndex];
+			}
+		}
+	}
+
+	else
+	{
+		// Only supporting the use of this function on RenderingThread.
+		check(false);
+		return;
+	}
+
+	if (LayerFound && LayerFound->Texture)
+	{
+		switch (LayerFound->ShapeType)
+		{
+		case IStereoLayers::CubemapLayer:
+			Texture = LayerFound->Texture->GetTextureCube();
+			LeftTexture = LayerFound->LeftTexture ? LayerFound->LeftTexture->GetTextureCube() : nullptr;			
+			break;
+
+		case IStereoLayers::CylinderLayer:
+		case IStereoLayers::QuadLayer:
+			Texture = LayerFound->Texture->GetTexture2D();
+			LeftTexture = LayerFound->LeftTexture ? LayerFound->LeftTexture->GetTexture2D() : nullptr;
+			break;
+
+		default:
+			break;
+		}
+	}
+}
diff --git a/Engine/Source/Runtime/HeadMountedDisplay/Public/DefaultSpectatorScreenController.h b/Engine/Source/Runtime/HeadMountedDisplay/Public/DefaultSpectatorScreenController.h
--- a/Engine/Source/Runtime/HeadMountedDisplay/Public/DefaultSpectatorScreenController.h
+++ b/Engine/Source/Runtime/HeadMountedDisplay/Public/DefaultSpectatorScreenController.h
@@ -25,6 +25,7 @@ public:
 	virtual void SetSpectatorScreenTexture(UTexture* InTexture) override;
 	virtual UTexture* GetSpectatorScreenTexture() const override;
 	virtual void SetSpectatorScreenModeTexturePlusEyeLayout(const FSpectatorScreenModeTexturePlusEyeLayout& Layout) override;
+	virtual void QueueDebugCanvasLayerID(int32 LayerID) override;
 
 	FSpectatorScreenRenderDelegate* GetSpectatorScreenRenderDelegate_RenderThread();
 
@@ -32,7 +33,7 @@ public:
 	// Implementation methods called by HMD
 	virtual void BeginRenderViewFamily();
 	virtual void UpdateSpectatorScreenMode_RenderThread();
-	virtual void RenderSpectatorScreen_RenderThread(FRHICommandListImmediate& RHICmdList, FRHITexture2D* BackBuffer, FTexture2DRHIRef RenderTarget, FVector2D WindowSize) const;
+	virtual void RenderSpectatorScreen_RenderThread(FRHICommandListImmediate& RHICmdList, FRHITexture2D* BackBuffer, FTexture2DRHIRef RenderTarget, FVector2D WindowSize);
 
 protected:
 	friend struct FRHISetSpectatorScreenTexture;
@@ -63,6 +64,7 @@ protected:
 	FTexture2DRHIRef SpectatorScreenTexture_RenderThread;
 	FSpectatorScreenModeTexturePlusEyeLayout SpectatorScreenModeTexturePlusEyeLayout_RenderThread;
 	FSpectatorScreenRenderDelegate SpectatorScreenDelegate_RenderThread;
+	TArray<int32> DebugCanvasLayerIDs;
 
 	class HEADMOUNTEDDISPLAY_API Helpers
 	{
diff --git a/Engine/Source/Runtime/HeadMountedDisplay/Public/DefaultStereoLayers.h b/Engine/Source/Runtime/HeadMountedDisplay/Public/DefaultStereoLayers.h
--- a/Engine/Source/Runtime/HeadMountedDisplay/Public/DefaultStereoLayers.h
+++ b/Engine/Source/Runtime/HeadMountedDisplay/Public/DefaultStereoLayers.h
@@ -26,6 +26,8 @@ public:
 	virtual void PostRenderView_RenderThread(FRHICommandListImmediate& RHICmdList, FSceneView& InView) override;
 	virtual bool IsActiveThisFrame(class FViewport* InViewport) const override;
 
+	void GetAllocatedTexture(uint32 LayerId, FTextureRHIRef &Texture, FTextureRHIRef &LeftTexture);
+
 protected:
 	
 	/**
diff --git a/Engine/Source/Runtime/HeadMountedDisplay/Public/HeadMountedDisplayBase.h b/Engine/Source/Runtime/HeadMountedDisplay/Public/HeadMountedDisplayBase.h
--- a/Engine/Source/Runtime/HeadMountedDisplay/Public/HeadMountedDisplayBase.h
+++ b/Engine/Source/Runtime/HeadMountedDisplay/Public/HeadMountedDisplayBase.h
@@ -49,7 +49,7 @@ public:
 	// Get the rectangle of the HMD rendertarget for the left eye which seems undistorted enough to be cropped and displayed on the spectator screen.
 	virtual FIntRect GetFullFlatEyeRect_RenderThread(FTexture2DRHIRef EyeTexture) const { return FIntRect(0, 0, 1, 1); }
 	// Helper to copy one render target into another for spectator screen display
-	virtual void CopyTexture_RenderThread(FRHICommandListImmediate& RHICmdList, FTexture2DRHIParamRef SrcTexture, FIntRect SrcRect, FTexture2DRHIParamRef DstTexture, FIntRect DstRect, bool bClearBlack) const {}
+	virtual void CopyTexture_RenderThread(FRHICommandListImmediate& RHICmdList, FTexture2DRHIParamRef SrcTexture, FIntRect SrcRect, FTexture2DRHIParamRef DstTexture, FIntRect DstRect, bool bClearBlack, bool bNoAlpha) const {}
 
 protected:
 	/**
diff --git a/Engine/Source/Runtime/HeadMountedDisplay/Public/IHeadMountedDisplay.h b/Engine/Source/Runtime/HeadMountedDisplay/Public/IHeadMountedDisplay.h
--- a/Engine/Source/Runtime/HeadMountedDisplay/Public/IHeadMountedDisplay.h
+++ b/Engine/Source/Runtime/HeadMountedDisplay/Public/IHeadMountedDisplay.h
@@ -115,6 +115,9 @@ public:
 	virtual void BeginRendering_GameThread() {}
 
 
+	// Are we outputting so a Spectator Screen now.
+	virtual bool IsSpectatorScreenActive() const { return false; }
+
 	/**
 	* Return a pointer to the SpectatorScreenController for the HMD if supported, else null.
 	* The controller is owned by the HMD, and will be destroyed when the HMD is destroyed.
diff --git a/Engine/Source/Runtime/HeadMountedDisplay/Public/ISpectatorScreenController.h b/Engine/Source/Runtime/HeadMountedDisplay/Public/ISpectatorScreenController.h
--- a/Engine/Source/Runtime/HeadMountedDisplay/Public/ISpectatorScreenController.h
+++ b/Engine/Source/Runtime/HeadMountedDisplay/Public/ISpectatorScreenController.h
@@ -54,4 +54,6 @@ public:
 	* @param	bClearBlack: if true the render target will be drawn black before either rect is drawn.
 	*/
 	virtual void SetSpectatorScreenModeTexturePlusEyeLayout(const FSpectatorScreenModeTexturePlusEyeLayout& Layout) = 0;
+
+	virtual void QueueDebugCanvasLayerID(int32 LayerID) = 0;
 };
diff --git a/Engine/Source/Runtime/HeadMountedDisplay/Public/StereoLayerManager.h b/Engine/Source/Runtime/HeadMountedDisplay/Public/StereoLayerManager.h
--- a/Engine/Source/Runtime/HeadMountedDisplay/Public/StereoLayerManager.h
+++ b/Engine/Source/Runtime/HeadMountedDisplay/Public/StereoLayerManager.h
@@ -14,6 +14,8 @@
 	The template argument should be a type for storing layer data. It should have a constructor matching the following:
 		LayerType(const FLayerDesc& InLayerDesc);
 	... and implement the following function overloads:
+		void SetLayerId(uint32 InId),
+		uint32 GetLayerId() const,
 		bool GetLayerDescMember(LayerType& Layer, FLayerDesc& OutLayerDesc),
 		void SetLayerDescMember(LayerType& Layer, const FLayerDesc& Desc), and
 		void MarkLayerTextureForUpdate(LayerType& Layer)
@@ -100,6 +102,7 @@ public:
 		uint32 LayerId = NextLayerId++;
 		check(LayerId > 0);
 		LayerType& NewLayer = StereoLayers.Emplace(LayerId, InLayerDesc);
+		NewLayer.SetLayerId(LayerId);
 		UpdateLayer(NewLayer, LayerId, InLayerDesc.Texture != nullptr);
 		bStereoLayersDirty = true;
 		return LayerId;
@@ -129,6 +132,7 @@ public:
 
 		LayerType& Layer = StereoLayers[LayerId];
 		SetLayerDescMember(Layer, InLayerDesc);
+		Layer.SetLayerId(LayerId);
 		UpdateLayer(Layer, LayerId, InLayerDesc.Texture != nullptr);
 		bStereoLayersDirty = true;
 	}
