Fixing an inadvertedly inverted condition and a later assert in relation to #jira UE-64644

diff --git a/Engine/Plugins/Runtime/Steam/SteamVR/Source/SteamVR/Private/SteamVRHMD.cpp b/Engine/Plugins/Runtime/Steam/SteamVR/Source/SteamVR/Private/SteamVRHMD.cpp
--- a/Engine/Plugins/Runtime/Steam/SteamVR/Source/SteamVR/Private/SteamVRHMD.cpp
+++ b/Engine/Plugins/Runtime/Steam/SteamVR/Source/SteamVR/Private/SteamVRHMD.cpp
@@ -1359,7 +1359,7 @@ void FSteamVRHMD::CalculateStereoViewOffset(const enum EStereoscopicPass StereoP
 
 FMatrix FSteamVRHMD::GetStereoProjectionMatrix(const enum EStereoscopicPass StereoPassType) const
 {
-	check(IsStereoEnabled());
+	check(IsStereoEnabled() || IsHeadTrackingEnforced());
 
 	vr::Hmd_Eye HmdEye = (StereoPassType == eSSP_LEFT_EYE) ? vr::Eye_Left : vr::Eye_Right;
 	float Left, Right, Top, Bottom;
diff --git a/Engine/Source/Runtime/HeadMountedDisplay/Private/HeadMountedDisplayBase.cpp b/Engine/Source/Runtime/HeadMountedDisplay/Private/HeadMountedDisplayBase.cpp
--- a/Engine/Source/Runtime/HeadMountedDisplay/Private/HeadMountedDisplayBase.cpp
+++ b/Engine/Source/Runtime/HeadMountedDisplay/Private/HeadMountedDisplayBase.cpp
@@ -97,7 +97,7 @@ FVector2D FHeadMountedDisplayBase::GetEyeCenterPoint_RenderThread(EStereoscopicP
 
 	// Note: IsHeadTrackingAllowed() can only be called from the game thread.
 	// IsStereoEnabled() and IsHeadTrackingEnforced() can be called from both the render and game threads, however.
-	if (IsStereoEnabled() || IsHeadTrackingEnforced())
+	if (!(IsStereoEnabled() || IsHeadTrackingEnforced()))
 	{
 		return FVector2D(0.5f, 0.5f);
 	}
