Duplicating from Release-4.21: Change: 4567465 UE-66353 Vulkan Crash on Startup
-PreLoadScreenManager starts up a vulkan swap chain before we Load UGeneralProjectSettings, but vulkan swap chain creation tries to read a setting from there and asserts.
-Now we just pretend the setting is false for the PreLoadScreenManager swap chain creation, which means vulkan may be vsynced (incorrect for VR).  A new swap chain is created after the load and does respect this setting.
-PreLoadScreenManager's functionality is unknown in VR, I am able to run a game that goes through this sequence.

Files:
//UE4/Release-4.21/Engine/Source/Runtime/HeadMountedDisplay/Private/HeadMountedDisplayVulkanExtensions.cpp#2

diff --git a/Engine/Source/Runtime/HeadMountedDisplay/Private/HeadMountedDisplayVulkanExtensions.cpp b/Engine/Source/Runtime/HeadMountedDisplay/Private/HeadMountedDisplayVulkanExtensions.cpp
--- a/Engine/Source/Runtime/HeadMountedDisplay/Private/HeadMountedDisplayVulkanExtensions.cpp
+++ b/Engine/Source/Runtime/HeadMountedDisplay/Private/HeadMountedDisplayVulkanExtensions.cpp
@@ -6,5 +6,8 @@
 
 bool IHeadMountedDisplayVulkanExtensions::ShouldDisableVulkanVSync() const
 {
-	return FParse::Param(FCommandLine::Get(), TEXT("vr")) || GetDefault<UGeneralProjectSettings>()->bStartInVR;
+	// This can be called while setting up a vulkan swap chain for the PreLoadScreenManager but we cannot easily inspect bStartInVR then because the object has not been loaded (it being PreLoad).
+	// For now we are just going to return false here and avoid an assert when trying to GetDefault.  This might break VR rendering of PreLoadScreens... but there is a good chance it is already broken.
+	// We check for the class being loaded first so that -vr behavior is consistent with bStartInVR.
+	return IsClassLoaded<UGeneralProjectSettings>() && (FParse::Param(FCommandLine::Get(), TEXT("vr")) || GetDefault<UGeneralProjectSettings>()->bStartInVR);
 }
