Use GetFirstGamePlayer to determine the primary XR player
This skips having to iterate all players and is much safer given that we can pass in the proper world context.

diff --git a/Engine/Source/Runtime/HeadMountedDisplay/Private/XRTrackingSystemBase.cpp b/Engine/Source/Runtime/HeadMountedDisplay/Private/XRTrackingSystemBase.cpp
--- a/Engine/Source/Runtime/HeadMountedDisplay/Private/XRTrackingSystemBase.cpp
+++ b/Engine/Source/Runtime/HeadMountedDisplay/Private/XRTrackingSystemBase.cpp
@@ -135,17 +135,8 @@ FTransform FXRTrackingSystemBase::ComputeTrackingToWorldTransform(FWorldContext&
 	UWorld* World = WorldContext.World();
 	if (World)
 	{
-		const ULocalPlayer* XRPlayer = nullptr;
-
-		const TArray<ULocalPlayer*>& LocalPlayers = GEngine->GetGamePlayers(World);
-		for (const ULocalPlayer* Player : LocalPlayers)
-		{
-			if (Player->IsPrimaryPlayer())
-			{
-				XRPlayer = Player;
-				break;
-			}
-		}
+		// Get the primary player for this world context
+		const ULocalPlayer* XRPlayer = GEngine->GetFirstGamePlayer(World);
 
 		if (XRPlayer)
 		{
