Fixed motion controllers only checking one layer deep for net ownership which prevented arbitrary actors from using motion controller components in a networked environment

diff --git a/Engine/Source/Runtime/HeadMountedDisplay/Private/MotionControllerComponent.cpp b/Engine/Source/Runtime/HeadMountedDisplay/Private/MotionControllerComponent.cpp
--- a/Engine/Source/Runtime/HeadMountedDisplay/Private/MotionControllerComponent.cpp
+++ b/Engine/Source/Runtime/HeadMountedDisplay/Private/MotionControllerComponent.cpp
@@ -471,8 +471,7 @@ bool UMotionControllerComponent::PollControllerState(FVector& Position, FRotator
 	{
 		// Cache state from the game thread for use on the render thread
 		const AActor* MyOwner = GetOwner();
-		const APawn* MyPawn = Cast<APawn>(MyOwner);
-		bHasAuthority = MyPawn ? MyPawn->IsLocallyControlled() : (MyOwner->Role == ENetRole::ROLE_Authority);
+		bHasAuthority = MyOwner->HasNetOwner();
 	}
 
 	if(bHasAuthority)
