Allow using external textures on Default Stereo Layers on GLES.

diff --git a/Engine/Shaders/Private/StereoLayerShader.usf b/Engine/Shaders/Private/StereoLayerShader.usf
--- a/Engine/Shaders/Private/StereoLayerShader.usf
+++ b/Engine/Shaders/Private/StereoLayerShader.usf
@@ -32,10 +32,20 @@ void MainVS(
 Texture2D InTexture;
 SamplerState InTextureSampler;
 
-void MainPS(
+void MainPS_Texture2D(
 	    float2 UV       : TEXCOORD0,
 	out float4 OutColor : SV_Target0
 	)
 {
     OutColor = Texture2DSample(InTexture, InTextureSampler, UV);
-}
\ No newline at end of file
+}
+
+TextureExternal InExternalTexture;
+
+void MainPS_TextureExternal(
+	    float2 UV       : TEXCOORD0,
+	out float4 OutColor : SV_Target0
+	)
+{
+    OutColor = TextureExternalSample(InExternalTexture, InTextureSampler, UV);
+}
diff --git a/Engine/Source/Runtime/Engine/Private/Components/StereoLayerComponent.cpp b/Engine/Source/Runtime/Engine/Private/Components/StereoLayerComponent.cpp
--- a/Engine/Source/Runtime/Engine/Private/Components/StereoLayerComponent.cpp
+++ b/Engine/Source/Runtime/Engine/Private/Components/StereoLayerComponent.cpp
@@ -111,6 +111,7 @@ void UStereoLayerComponent::TickComponent(float DeltaTime, enum ELevelTick TickT
 			LayerDsec.Flags |= (bNoAlphaChannel) ? IStereoLayers::LAYER_FLAG_TEX_NO_ALPHA_CHANNEL : 0;
 			LayerDsec.Flags |= (bQuadPreserveTextureRatio) ? IStereoLayers::LAYER_FLAG_QUAD_PRESERVE_TEX_RATIO : 0;
 			LayerDsec.Flags |= (bSupportsDepth) ? IStereoLayers::LAYER_FLAG_SUPPORT_DEPTH : 0;
+			LayerDsec.Flags |= (Texture->GetMaterialType() == MCT_TextureExternal) ? IStereoLayers::LAYER_FLAG_TEX_EXTERNAL : 0;
 
 			switch (StereoLayerType)
 			{
diff --git a/Engine/Source/Runtime/Engine/Private/StereoLayerRendering.cpp b/Engine/Source/Runtime/Engine/Private/StereoLayerRendering.cpp
--- a/Engine/Source/Runtime/Engine/Private/StereoLayerRendering.cpp
+++ b/Engine/Source/Runtime/Engine/Private/StereoLayerRendering.cpp
@@ -7,4 +7,5 @@
 #include "StereoLayerRendering.h"
 
 IMPLEMENT_SHADER_TYPE(,FStereoLayerVS,TEXT("/Engine/Private/StereoLayerShader.usf"),TEXT("MainVS"),SF_Vertex);
-IMPLEMENT_SHADER_TYPE(,FStereoLayerPS,TEXT("/Engine/Private/StereoLayerShader.usf"),TEXT("MainPS"),SF_Pixel);
+IMPLEMENT_SHADER_TYPE(,FStereoLayerPS,TEXT("/Engine/Private/StereoLayerShader.usf"),TEXT("MainPS_Texture2D"),SF_Pixel);
+IMPLEMENT_SHADER_TYPE(,FStereoLayerPS_External,TEXT("/Engine/Private/StereoLayerShader.usf"),TEXT("MainPS_TextureExternal"),SF_Pixel);
diff --git a/Engine/Source/Runtime/Engine/Public/IStereoLayers.h b/Engine/Source/Runtime/Engine/Public/IStereoLayers.h
--- a/Engine/Source/Runtime/Engine/Public/IStereoLayers.h
+++ b/Engine/Source/Runtime/Engine/Public/IStereoLayers.h
@@ -37,7 +37,9 @@ public:
 		// Quad Y component will be calculated based on the texture dimensions
 		LAYER_FLAG_QUAD_PRESERVE_TEX_RATIO	= 0x00000004,
 		// The layer will intersect with the scene's depth. Currently only supported on Oculus platforms.
-		LAYER_FLAG_SUPPORT_DEPTH			= 0x00000008,
+		LAYER_FLAG_SUPPORT_DEPTH = 0x00000008,
+		// Required on some platforms to enable rendering of external textures.
+		LAYER_FLAG_TEX_EXTERNAL = 0x00000010,
 	};
 
 	/**
diff --git a/Engine/Source/Runtime/Engine/Public/StereoLayerRendering.h b/Engine/Source/Runtime/Engine/Public/StereoLayerRendering.h
--- a/Engine/Source/Runtime/Engine/Public/StereoLayerRendering.h
+++ b/Engine/Source/Runtime/Engine/Public/StereoLayerRendering.h
@@ -79,29 +79,17 @@ private:
 	FShaderParameter InWorld;
 };
 
-/**
- * A pixel shader for rendering a transformed textured element.
- */
-class FStereoLayerPS : public FGlobalShader
+class FStereoLayerPS_Base : public FGlobalShader
 {
-	DECLARE_EXPORTED_SHADER_TYPE(FStereoLayerPS,Global,ENGINE_API);
 public:
 
 	static bool ShouldCompilePermutation(const FGlobalShaderPermutationParameters& Parameters) { return true; }
 
-	FStereoLayerPS(const ShaderMetaType::CompiledShaderInitializerType& Initializer):
-		FGlobalShader(Initializer)
-	{
-		InTexture.Bind(Initializer.ParameterMap,TEXT("InTexture"), SPF_Mandatory);
-		InTextureSampler.Bind(Initializer.ParameterMap,TEXT("InTextureSampler"));
-	}
-	FStereoLayerPS() {}
-
 	void SetParameters(FRHICommandList& RHICmdList, FSamplerStateRHIParamRef SamplerStateRHI, FTextureRHIParamRef TextureRHI)
 	{
 		FPixelShaderRHIParamRef PS = GetPixelShader();
 
-		SetTextureParameter(RHICmdList, PS,InTexture,InTextureSampler,SamplerStateRHI,TextureRHI);
+		SetTextureParameter(RHICmdList, PS, InTexture, InTextureSampler, SamplerStateRHI, TextureRHI);
 	}
 
 	virtual bool Serialize(FArchive& Ar) override
@@ -112,7 +100,46 @@ public:
 		return bShaderHasOutdatedParameters;
 	}
 
-private:
+protected:
 	FShaderResourceParameter InTexture;
 	FShaderResourceParameter InTextureSampler;
+
+	FStereoLayerPS_Base(const ShaderMetaType::CompiledShaderInitializerType& Initializer, const TCHAR* TextureParamName) :
+		FGlobalShader(Initializer) 
+	{
+		InTexture.Bind(Initializer.ParameterMap, TextureParamName, SPF_Mandatory);
+		InTextureSampler.Bind(Initializer.ParameterMap, TEXT("InTextureSampler"));
+	}
+	FStereoLayerPS_Base() {}
+
+};
+
+/**
+ * A pixel shader for rendering a transformed textured element.
+ */
+class FStereoLayerPS : public FStereoLayerPS_Base
+{
+	DECLARE_EXPORTED_SHADER_TYPE(FStereoLayerPS,Global,ENGINE_API);
+public:
+
+	static bool ShouldCompilePermutation(const FGlobalShaderPermutationParameters& Parameters) { return true; }
+
+	FStereoLayerPS(const ShaderMetaType::CompiledShaderInitializerType& Initializer):
+		FStereoLayerPS_Base(Initializer, TEXT("InTexture"))	{}
+	FStereoLayerPS() {}
+};
+
+/**
+ * A pixel shader for rendering a transformed external texture element.
+ */
+class FStereoLayerPS_External : public FStereoLayerPS_Base
+{
+	DECLARE_EXPORTED_SHADER_TYPE(FStereoLayerPS_External, Global, ENGINE_API);
+public:
+
+	static bool ShouldCompilePermutation(const FGlobalShaderPermutationParameters& Parameters) { return true; }
+
+	FStereoLayerPS_External(const ShaderMetaType::CompiledShaderInitializerType& Initializer) :
+		FStereoLayerPS_Base(Initializer, TEXT("InExternalTexture")) {}
+	FStereoLayerPS_External() {}
 };
diff --git a/Engine/Source/Runtime/HeadMountedDisplay/Private/DefaultStereoLayers.cpp b/Engine/Source/Runtime/HeadMountedDisplay/Private/DefaultStereoLayers.cpp
--- a/Engine/Source/Runtime/HeadMountedDisplay/Private/DefaultStereoLayers.cpp
+++ b/Engine/Source/Runtime/HeadMountedDisplay/Private/DefaultStereoLayers.cpp
@@ -63,55 +63,60 @@ void FDefaultStereoLayers::StereoLayerRender(FRHICommandListImmediate& RHICmdLis
 	}
 
 	IRendererModule& RendererModule = GetRendererModule();
+	using TOpaqueBlendState = TStaticBlendState<CW_RGBA, BO_Add, BF_One, BF_Zero, BO_Add, BF_One, BF_Zero>;
+	using TAlphaBlendState = TStaticBlendState<CW_RGBA, BO_Add, BF_SourceAlpha, BF_InverseSourceAlpha, BO_Add, BF_One, BF_InverseSourceAlpha>;
 
 	// Set render state
 	FGraphicsPipelineStateInitializer GraphicsPSOInit;
 	RHICmdList.ApplyCachedRenderTargets(GraphicsPSOInit);
-	bool bLastNoAlpha = (RenderThreadLayers[LayersToRender[0]].Flags & LAYER_FLAG_TEX_NO_ALPHA_CHANNEL) != 0;
-	if (bLastNoAlpha)
-	{
-		GraphicsPSOInit.BlendState = TStaticBlendState<CW_RGBA, BO_Add, BF_One, BF_Zero, BO_Add, BF_One, BF_Zero>::GetRHI();
-	}
-	else
-	{
-		GraphicsPSOInit.BlendState = TStaticBlendState<CW_RGBA, BO_Add, BF_SourceAlpha, BF_InverseSourceAlpha, BO_Add, BF_One, BF_InverseSourceAlpha>::GetRHI();
-	}
 
 	GraphicsPSOInit.RasterizerState = TStaticRasterizerState<FM_Solid, CM_None, true, false>::GetRHI();
 	GraphicsPSOInit.DepthStencilState = TStaticDepthStencilState<false, CF_Always>::GetRHI();
 	RHICmdList.SetScissorRect(false, 0, 0, 0, 0);
 	RHICmdList.SetViewport(RenderParams.Viewport.Min.X, RenderParams.Viewport.Min.Y, 0, RenderParams.Viewport.Max.X, RenderParams.Viewport.Max.Y, 1.0f);
 
-	// Set shader state
+	// Set initial shader state
 	auto ShaderMap = GetGlobalShaderMap(GMaxRHIFeatureLevel);
 	TShaderMapRef<FStereoLayerVS> VertexShader(ShaderMap);
 	TShaderMapRef<FStereoLayerPS> PixelShader(ShaderMap);
+	TShaderMapRef<FStereoLayerPS_External> PixelShader_External(ShaderMap);
 
 	GraphicsPSOInit.BoundShaderState.VertexDeclarationRHI = GFilterVertexDeclaration.VertexDeclarationRHI;
 	GraphicsPSOInit.BoundShaderState.VertexShaderRHI = GETSAFERHISHADER_VERTEX(*VertexShader);
-	GraphicsPSOInit.BoundShaderState.PixelShaderRHI = GETSAFERHISHADER_PIXEL(*PixelShader);
+
 	GraphicsPSOInit.PrimitiveType = PT_TriangleList;
-	SetGraphicsPipelineState(RHICmdList, GraphicsPSOInit);
+
+	// Force initialization of pipeline state on first iteration:
+	bool bLastWasOpaque = (RenderThreadLayers[LayersToRender[0]].Flags & LAYER_FLAG_TEX_NO_ALPHA_CHANNEL) == 0;
+	bool bLastWasExternal = (RenderThreadLayers[LayersToRender[0]].Flags & LAYER_FLAG_TEX_EXTERNAL) == 0;
 
 	// For each layer
 	for (uint32 LayerIndex : LayersToRender)
 	{
 		const FLayerDesc& Layer = RenderThreadLayers[LayerIndex];
 		check(Layer.Texture.IsValid());
-		const bool bNoAlpha = (Layer.Flags & LAYER_FLAG_TEX_NO_ALPHA_CHANNEL) != 0;
-		if (bNoAlpha != bLastNoAlpha)
+		const bool bIsOpaque = (Layer.Flags & LAYER_FLAG_TEX_NO_ALPHA_CHANNEL) != 0;
+		const bool bIsExternal = (Layer.Flags & LAYER_FLAG_TEX_EXTERNAL) != 0;
+		bool bPipelineStateNeedsUpdate = false;
+
+		if (bIsOpaque != bLastWasOpaque)
+		{
+			bLastWasOpaque = bIsOpaque;
+			GraphicsPSOInit.BlendState = bIsOpaque ? TOpaqueBlendState::GetRHI() : TAlphaBlendState::GetRHI();
+			bPipelineStateNeedsUpdate = true;
+		}
+
+		if (bIsExternal != bLastWasExternal)
+		{
+			bLastWasExternal = bIsExternal;
+			GraphicsPSOInit.BoundShaderState.PixelShaderRHI = bIsExternal ? GETSAFERHISHADER_PIXEL(*PixelShader_External) : GETSAFERHISHADER_PIXEL(*PixelShader);
+			bPipelineStateNeedsUpdate = true;
+		}
+
+		if (bPipelineStateNeedsUpdate)
 		{
 			// Updater render state
-			if (bNoAlpha)
-			{
-				GraphicsPSOInit.BlendState = TStaticBlendState<CW_RGBA, BO_Add, BF_One, BF_Zero, BO_Add, BF_InverseSourceAlpha, BF_Zero>::GetRHI();
-			}
-			else
-			{
-				GraphicsPSOInit.BlendState = TStaticBlendState<CW_RGBA, BO_Add, BF_SourceAlpha, BF_InverseSourceAlpha, BO_Add, BF_One, BF_InverseSourceAlpha>::GetRHI();
-			}
 			SetGraphicsPipelineState(RHICmdList, GraphicsPSOInit);
-			bLastNoAlpha = bNoAlpha;
 		}
 
 		FMatrix LayerMatrix = ConvertTransform(Layer.Transform);
