Fix build server build error in XRSwapChain.cpp (missing #include)

diff --git a/Engine/Source/Runtime/HeadMountedDisplay/Private/XRSwapChain.cpp b/Engine/Source/Runtime/HeadMountedDisplay/Private/XRSwapChain.cpp
--- a/Engine/Source/Runtime/HeadMountedDisplay/Private/XRSwapChain.cpp
+++ b/Engine/Source/Runtime/HeadMountedDisplay/Private/XRSwapChain.cpp
@@ -1,6 +1,7 @@
 // Copyright 1998-2019 Epic Games, Inc. All Rights Reserved.
 
 #include "XRSwapChain.h"
+#include "RenderingThread.h"
 #include "XRThreadUtils.h"
 #include "CoreGlobals.h"
 
