There are two minor issues with the XRSwapChain implementation:

On initialization we're creating a temp copy of the TArray of textures; this would be more efficient to pass by RValue-reference (move) to avoid the copy.
On shutdown, can hit issues with cleanup code running on the correct thread (should be run on RHI thread, and needs special case if running from game thread).

diff --git a/Engine/Plugins/Runtime/Oculus/OculusVR/Source/OculusHMD/Private/OculusHMD_CustomPresent.cpp b/Engine/Plugins/Runtime/Oculus/OculusVR/Source/OculusHMD/Private/OculusHMD_CustomPresent.cpp
--- a/Engine/Plugins/Runtime/Oculus/OculusVR/Source/OculusHMD/Private/OculusHMD_CustomPresent.cpp
+++ b/Engine/Plugins/Runtime/Oculus/OculusVR/Source/OculusHMD/Private/OculusHMD_CustomPresent.cpp
@@ -293,7 +293,7 @@ FXRSwapChainPtr FCustomPresent::CreateSwapChain_RenderThread(uint32 InSizeX, uin
 		}
 	}
 
-	return MakeShareable(new FXRSwapChain(RHITexture, RHITextureSwapChain));
+	return MakeShareable(new FXRSwapChain(RHITexture, MoveTemp(RHITextureSwapChain)));
 }
 
 
diff --git a/Engine/Source/Runtime/HeadMountedDisplay/Private/XRSwapChain.cpp b/Engine/Source/Runtime/HeadMountedDisplay/Private/XRSwapChain.cpp
--- a/Engine/Source/Runtime/HeadMountedDisplay/Private/XRSwapChain.cpp
+++ b/Engine/Source/Runtime/HeadMountedDisplay/Private/XRSwapChain.cpp
@@ -60,7 +60,7 @@ FORCEINLINE void CheckInRHIThread()
 #endif
 }
 
-FXRSwapChain::FXRSwapChain(FTextureRHIParamRef InRHITexture, const TArray<FTextureRHIRef>& InRHITextureSwapChain)
+FXRSwapChain::FXRSwapChain(FTextureRHIParamRef InRHITexture, TArray<FTextureRHIRef>&& InRHITextureSwapChain)
 	: RHITexture(InRHITexture)
 	, RHITextureSwapChain(InRHITextureSwapChain)
 	, SwapChainIndex_RHIThread(0)
@@ -70,15 +70,23 @@ FXRSwapChain::FXRSwapChain(FTextureRHIParamRef InRHITexture, const TArray<FTextu
 
 FXRSwapChain::~FXRSwapChain()
 {
-	// @todo: The assertion below fires on exit (even though we're on a valid thread) I'm guessing because everything's getting torn 
-	// down and our thread checking isn't valid. Not sure whether to remove the assert here, or figure out a more correct way to 
-	// perform the check.
-//	CheckInRHIThread();
-
-	ExecuteOnRHIThread([this]()
+	if (IsInGameThread())
+	{
+		ExecuteOnRenderThread([this]()
+		{
+			ExecuteOnRHIThread([this]()
+			{
+				ReleaseResources_RHIThread();
+			});
+		});
+	}
+	else
 	{
-		ReleaseResources_RHIThread();
-	});
+		ExecuteOnRHIThread([this]()
+		{
+			ReleaseResources_RHIThread();
+		});
+	}
 }
 
 
diff --git a/Engine/Source/Runtime/HeadMountedDisplay/Public/XRSwapChain.h b/Engine/Source/Runtime/HeadMountedDisplay/Public/XRSwapChain.h
--- a/Engine/Source/Runtime/HeadMountedDisplay/Public/XRSwapChain.h
+++ b/Engine/Source/Runtime/HeadMountedDisplay/Public/XRSwapChain.h
@@ -11,7 +11,7 @@
 class HEADMOUNTEDDISPLAY_API FXRSwapChain : public TSharedFromThis<FXRSwapChain, ESPMode::ThreadSafe>
 {
 public:
-	FXRSwapChain(FTextureRHIParamRef InRHITexture, const TArray<FTextureRHIRef>& InRHITextureSwapChain);
+	FXRSwapChain(FTextureRHIParamRef InRHITexture, TArray<FTextureRHIRef>&& InRHITextureSwapChain);
 	virtual ~FXRSwapChain();
 
 	FRHITexture* GetTexture() const { return RHITexture.GetReference(); }
