FXRTrackingSystemBase::GetRelativeEyePose eyes backward in y
-We had the left eye on the right, and vice versa, which is great at causing eyestrain.  Most platforms override this.

diff --git a/Engine/Source/Runtime/HeadMountedDisplay/Private/XRTrackingSystemBase.cpp b/Engine/Source/Runtime/HeadMountedDisplay/Private/XRTrackingSystemBase.cpp
--- a/Engine/Source/Runtime/HeadMountedDisplay/Private/XRTrackingSystemBase.cpp
+++ b/Engine/Source/Runtime/HeadMountedDisplay/Private/XRTrackingSystemBase.cpp
@@ -63,7 +63,7 @@ bool FXRTrackingSystemBase::GetRelativeEyePose(int32 DeviceId, EStereoscopicPass
 	OutPosition = FVector::ZeroVector;
 	if (DeviceId == IXRTrackingSystem::HMDDeviceId && (Eye == eSSP_LEFT_EYE || Eye == eSSP_RIGHT_EYE))
 	{
-		OutPosition = FVector(0, (Eye == EStereoscopicPass::eSSP_LEFT_EYE ? .5 : -.5) * 0.064f * GetWorldToMetersScale(), 0);
+		OutPosition = FVector(0, (Eye == EStereoscopicPass::eSSP_LEFT_EYE ? -.5 : .5) * 0.064f * GetWorldToMetersScale(), 0);
 		return true;
 	}
 	else
