Fixing Oculus mobile crash

diff --git a/Engine/Source/Runtime/HeadMountedDisplay/Private/LateUpdateManager.cpp b/Engine/Source/Runtime/HeadMountedDisplay/Private/LateUpdateManager.cpp
--- a/Engine/Source/Runtime/HeadMountedDisplay/Private/LateUpdateManager.cpp
+++ b/Engine/Source/Runtime/HeadMountedDisplay/Private/LateUpdateManager.cpp
@@ -5,7 +5,7 @@
 #include "Components/PrimitiveComponent.h"
 #include "PrimitiveSceneInfo.h"
 
-FLateUpdateManager::FLateUpdateManager() 
+FLateUpdateManager::FLateUpdateManager()
 	: LateUpdateGameWriteIndex(0)
 	, LateUpdateRenderReadIndex(0)
 {
@@ -102,7 +102,7 @@ void FLateUpdateManager::CacheSceneInfo(USceneComponent* Component)
 	if (PrimitiveComponent && PrimitiveComponent->SceneProxy)
 	{
 		FPrimitiveSceneInfo* PrimitiveSceneInfo = PrimitiveComponent->SceneProxy->GetPrimitiveSceneInfo();
-		if (PrimitiveSceneInfo)
+		if (PrimitiveSceneInfo && PrimitiveSceneInfo->IsIndexValid())
 		{
 			LateUpdatePrimitives[LateUpdateGameWriteIndex].Emplace(PrimitiveSceneInfo, PrimitiveSceneInfo->GetIndex());
 		}
@@ -122,4 +122,4 @@ void FLateUpdateManager::GatherLateUpdatePrimitives(USceneComponent* ParentCompo
 			CacheSceneInfo(Component);
 		}
 	}
-}
\ No newline at end of file
+}
