Simplifying GetHMDDevice and GetStereoRenderingDevice so each derived class doesn't need to duplicate functionality.

diff --git a/Engine/Plugins/Experimental/SimpleHMD/Source/SimpleHMD/Private/SimpleHMD.h b/Engine/Plugins/Experimental/SimpleHMD/Source/SimpleHMD/Private/SimpleHMD.h
--- a/Engine/Plugins/Experimental/SimpleHMD/Source/SimpleHMD/Private/SimpleHMD.h
+++ b/Engine/Plugins/Experimental/SimpleHMD/Source/SimpleHMD/Private/SimpleHMD.h
@@ -41,15 +41,6 @@ public:
 	virtual void SetBaseOrientation(const FQuat& BaseOrient) override;
 	virtual FQuat GetBaseOrientation() const override;
 
-	virtual class IHeadMountedDisplay* GetHMDDevice() override
-	{ 
-		return this;
-	}
-
-	virtual class TSharedPtr< class IStereoRendering, ESPMode::ThreadSafe > GetStereoRenderingDevice() override
-	{
-		return SharedThis(this);
-	}
 protected:
 	/** FXRTrackingSystemBase protected interface */
 	virtual float GetWorldToMetersScale() const override;
diff --git a/Engine/Plugins/Lumin/MagicLeap/Source/MagicLeap/Private/MagicLeapHMD.h b/Engine/Plugins/Lumin/MagicLeap/Source/MagicLeap/Private/MagicLeapHMD.h
--- a/Engine/Plugins/Lumin/MagicLeap/Source/MagicLeap/Private/MagicLeapHMD.h
+++ b/Engine/Plugins/Lumin/MagicLeap/Source/MagicLeap/Private/MagicLeapHMD.h
@@ -55,8 +55,6 @@ public:
 	virtual bool OnEndGameFrame(FWorldContext& WorldContext) override;
 	virtual void OnBeginRendering_GameThread() override;
 	virtual void OnBeginRendering_RenderThread(FRHICommandListImmediate& RHICmdList, FSceneViewFamily& ViewFamily) override;
-	virtual class IHeadMountedDisplay* GetHMDDevice() override { return this; }
-	virtual class TSharedPtr< class IStereoRendering, ESPMode::ThreadSafe > GetStereoRenderingDevice() override { return AsShared(); }
 	virtual class TSharedPtr< class IXRCamera, ESPMode::ThreadSafe > GetXRCamera(int32 DeviceId) override;
 	virtual FName GetSystemName() const override;
 	virtual FString GetVersionString() const override;
diff --git a/Engine/Plugins/Runtime/GoogleVR/GoogleVRHMD/Source/GoogleVRHMD/Private/GoogleVRHMD.h b/Engine/Plugins/Runtime/GoogleVR/GoogleVRHMD/Source/GoogleVRHMD/Private/GoogleVRHMD.h
--- a/Engine/Plugins/Runtime/GoogleVR/GoogleVRHMD/Source/GoogleVRHMD/Private/GoogleVRHMD.h
+++ b/Engine/Plugins/Runtime/GoogleVR/GoogleVRHMD/Source/GoogleVRHMD/Private/GoogleVRHMD.h
@@ -663,10 +663,6 @@ public:
 	*/
 	virtual class IXRInput* GetXRInput() { return this; }
 
-	virtual class IHeadMountedDisplay* GetHMDDevice() override { return this; }
-
-	virtual class TSharedPtr< class IStereoRendering, ESPMode::ThreadSafe > GetStereoRenderingDevice() override { return SharedThis(this); }
-
 	//////////////////////////////////////
 	// Begin IXRInput Virtual Interface //
 	//////////////////////////////////////
diff --git a/Engine/Plugins/Runtime/Oculus/OculusVR/Source/OculusHMD/Private/OculusHMD.h b/Engine/Plugins/Runtime/Oculus/OculusVR/Source/OculusHMD/Private/OculusHMD.h
--- a/Engine/Plugins/Runtime/Oculus/OculusVR/Source/OculusHMD/Private/OculusHMD.h
+++ b/Engine/Plugins/Runtime/Oculus/OculusVR/Source/OculusHMD/Private/OculusHMD.h
@@ -97,8 +97,6 @@ public:
 	virtual void SetBaseOrientation(const FQuat& BaseOrient) override;
 	virtual FQuat GetBaseOrientation() const override;
 	//virtual TSharedPtr<class IXRCamera, ESPMode::ThreadSafe> GetXRCamera(int32 DeviceId = HMDDeviceId) override;
-	virtual class IHeadMountedDisplay* GetHMDDevice() override { return this; }
-	virtual class TSharedPtr< class IStereoRendering, ESPMode::ThreadSafe > GetStereoRenderingDevice() override { return SharedThis(this); }
 	//virtual class IXRInput* GetXRInput() override;
 	virtual bool IsHeadTrackingEnforced() const override;
 	virtual void SetHeadTrackingEnforced(bool bEnabled) override;
diff --git a/Engine/Plugins/Runtime/OpenXR/Source/OpenXRHMD/Private/OpenXRHMD.h b/Engine/Plugins/Runtime/OpenXR/Source/OpenXRHMD/Private/OpenXRHMD.h
--- a/Engine/Plugins/Runtime/OpenXR/Source/OpenXRHMD/Private/OpenXRHMD.h
+++ b/Engine/Plugins/Runtime/OpenXR/Source/OpenXRHMD/Private/OpenXRHMD.h
@@ -96,15 +96,6 @@ public:
 		return (TrackingSpaceType == XR_REFERENCE_SPACE_TYPE_STAGE) ? EHMDTrackingOrigin::Floor : EHMDTrackingOrigin::Eye;
 	}
 
-	virtual class IHeadMountedDisplay* GetHMDDevice() override
-	{ 
-		return this;
-	}
-
-	virtual class TSharedPtr< class IStereoRendering, ESPMode::ThreadSafe > GetStereoRenderingDevice() override
-	{
-		return SharedThis(this);
-	}
 protected:
 	/** FXRTrackingSystemBase protected interface */
 	virtual float GetWorldToMetersScale() const override;
diff --git a/Engine/Plugins/Runtime/Steam/SteamVR/Source/SteamVR/Private/SteamVRHMD.h b/Engine/Plugins/Runtime/Steam/SteamVR/Source/SteamVR/Private/SteamVRHMD.h
--- a/Engine/Plugins/Runtime/Steam/SteamVR/Source/SteamVR/Private/SteamVRHMD.h
+++ b/Engine/Plugins/Runtime/Steam/SteamVR/Source/SteamVR/Private/SteamVRHMD.h
@@ -127,17 +127,6 @@ public:
 	}
 	virtual FString GetVersionString() const override;
 
-
-	virtual class IHeadMountedDisplay* GetHMDDevice() override
-	{
-		return this;
-	}
-
-	virtual class TSharedPtr< class IStereoRendering, ESPMode::ThreadSafe > GetStereoRenderingDevice() override
-	{
-		return SharedThis(this);
-	}
-
 	virtual bool OnStartGameFrame(FWorldContext& WorldContext) override;
 	virtual bool DoesSupportPositionalTracking() const override;
 	virtual bool HasValidTrackingPosition() override;
diff --git a/Engine/Plugins/Runtime/WindowsMixedReality/Source/WindowsMixedRealityHMD/Private/WindowsMixedRealityHMD.h b/Engine/Plugins/Runtime/WindowsMixedReality/Source/WindowsMixedRealityHMD/Private/WindowsMixedRealityHMD.h
--- a/Engine/Plugins/Runtime/WindowsMixedReality/Source/WindowsMixedRealityHMD/Private/WindowsMixedRealityHMD.h
+++ b/Engine/Plugins/Runtime/WindowsMixedReality/Source/WindowsMixedRealityHMD/Private/WindowsMixedRealityHMD.h
@@ -71,15 +71,6 @@ namespace WindowsMixedReality
 			FVector& CurrentPosition) override;
 		virtual bool GetRelativeEyePose(int32 DeviceId, EStereoscopicPass Eye, FQuat& OutOrientation, FVector& OutPosition) override;
 
-		virtual class IHeadMountedDisplay* GetHMDDevice() override
-		{
-			return this;
-		}
-
-		virtual class TSharedPtr< class IStereoRendering, ESPMode::ThreadSafe > GetStereoRenderingDevice() override
-		{
-			return SharedThis(this);
-		}
 	protected:
 		/** FXRTrackingSystemBase protected interface */
 		virtual float GetWorldToMetersScale() const override;
diff --git a/Engine/Source/Runtime/HeadMountedDisplay/Public/HeadMountedDisplayBase.h b/Engine/Source/Runtime/HeadMountedDisplay/Public/HeadMountedDisplayBase.h
--- a/Engine/Source/Runtime/HeadMountedDisplay/Public/HeadMountedDisplayBase.h
+++ b/Engine/Source/Runtime/HeadMountedDisplay/Public/HeadMountedDisplayBase.h
@@ -11,7 +11,7 @@
  * You can extend this class instead of IHeadMountedDisplay directly when implementing support for new HMD devices.
  */
 
-class HEADMOUNTEDDISPLAY_API FHeadMountedDisplayBase : public FXRTrackingSystemBase, public IHeadMountedDisplay, public IStereoRendering
+class HEADMOUNTEDDISPLAY_API FHeadMountedDisplayBase : public FXRTrackingSystemBase, public IHeadMountedDisplay, public IStereoRendering, public TSharedFromThis<FHeadMountedDisplayBase, ESPMode::ThreadSafe>
 {
 
 public:
@@ -57,6 +57,24 @@ public:
 	// Helper to copy one render target into another for spectator screen display
 	virtual void CopyTexture_RenderThread(FRHICommandListImmediate& RHICmdList, FTexture2DRHIParamRef SrcTexture, FIntRect SrcRect, FTexture2DRHIParamRef DstTexture, FIntRect DstRect, bool bClearBlack, bool bNoAlpha) const {}
 
+	/**
+	 * Access HMD rendering-related features.
+	 *
+	 * @return a IHeadmountedDisplay pointer or a nullptr if this tracking system does not support head mounted displays.
+	 */
+	virtual class IHeadMountedDisplay* GetHMDDevice() override { return this; }
+
+	/**
+	* Access Stereo rendering device associated with this XR system.
+	* If GetHMDDevice() returns non-null, this method should also return a vaild instance.
+	*
+	* @return a IStereoRendering pointer or a nullptr if this tracking system does not support stereo rendering.
+	*/
+	virtual class TSharedPtr< class IStereoRendering, ESPMode::ThreadSafe > GetStereoRenderingDevice() override
+	{
+		return SharedThis(this);
+	}
+
 protected:
 	/**
 	 * Called by RecordAnalytics when creating the analytics event sent during HMD initialization.
