Selective move of Dev-VR changes to base XRSwapChain class (and one client usage) to fix build break with new OpenXR plugin code in 4.23. Also fix initializer-list ordering warning in OpenXR plugin.

diff --git a/Engine/Source/Runtime/HeadMountedDisplay/Private/XRSwapChain.cpp b/Engine/Source/Runtime/HeadMountedDisplay/Private/XRSwapChain.cpp
--- a/Engine/Source/Runtime/HeadMountedDisplay/Private/XRSwapChain.cpp
+++ b/Engine/Source/Runtime/HeadMountedDisplay/Private/XRSwapChain.cpp
@@ -60,8 +60,7 @@ FORCEINLINE void CheckInRHIThread()
 #endif
 }
 
-FXRSwapChain::FXRSwapChain(TArray<FTextureRHIRef>&& InRHITextureSwapChain, const FTextureRHIRef & AliasedTexture)
-	: RHITexture(AliasedTexture)
+FXRSwapChain::FXRSwapChain(TArray<FTextureRHIRef>&& InRHITextureSwapChain, const FTextureRHIRef& AliasedTexture)
 	, RHITextureSwapChain(InRHITextureSwapChain)
 	, SwapChainIndex_RHIThread(0)
 {
