Fix merge errors on OpenXRHMD.h and XRSwapChain.cpp
One was a bad initializer list (initialiated with a "," instead of a ":"), the other an ambiguous TSharedFromThis reference since two base classes now derive from that.

diff --git a/Engine/Plugins/Runtime/OpenXR/Source/OpenXRHMD/Private/OpenXRHMD.h b/Engine/Plugins/Runtime/OpenXR/Source/OpenXRHMD/Private/OpenXRHMD.h
--- a/Engine/Plugins/Runtime/OpenXR/Source/OpenXRHMD/Private/OpenXRHMD.h
+++ b/Engine/Plugins/Runtime/OpenXR/Source/OpenXRHMD/Private/OpenXRHMD.h
@@ -93,7 +93,7 @@ public:
 	}
 	virtual class TSharedPtr< class IStereoRendering, ESPMode::ThreadSafe > GetStereoRenderingDevice() override
 	{
-		return SharedThis(this);
+		return FHeadMountedDisplayBase::SharedThis(static_cast<FHeadMountedDisplayBase*>(this));
 	}
 
 protected:
diff --git a/Engine/Source/Runtime/HeadMountedDisplay/Private/XRSwapChain.cpp b/Engine/Source/Runtime/HeadMountedDisplay/Private/XRSwapChain.cpp
--- a/Engine/Source/Runtime/HeadMountedDisplay/Private/XRSwapChain.cpp
+++ b/Engine/Source/Runtime/HeadMountedDisplay/Private/XRSwapChain.cpp
@@ -61,7 +61,7 @@ FORCEINLINE void CheckInRHIThread()
 }
 
 FXRSwapChain::FXRSwapChain(TArray<FTextureRHIRef>&& InRHITextureSwapChain, const FTextureRHIRef& AliasedTexture)
-	, RHITextureSwapChain(InRHITextureSwapChain)
+	: RHITextureSwapChain(InRHITextureSwapChain)
 	, SwapChainIndex_RHIThread(0)
 {
 	check(RHITexture);
