Fix bad merge from 4.23 - RHITexture member of XRSwapChain needs to be set as well as the array of textures.

diff --git a/Engine/Source/Runtime/HeadMountedDisplay/Private/XRSwapChain.cpp b/Engine/Source/Runtime/HeadMountedDisplay/Private/XRSwapChain.cpp
--- a/Engine/Source/Runtime/HeadMountedDisplay/Private/XRSwapChain.cpp
+++ b/Engine/Source/Runtime/HeadMountedDisplay/Private/XRSwapChain.cpp
@@ -61,7 +61,8 @@ FORCEINLINE void CheckInRHIThread()
 }
 
 FXRSwapChain::FXRSwapChain(TArray<FTextureRHIRef>&& InRHITextureSwapChain, const FTextureRHIRef& AliasedTexture)
-	: RHITextureSwapChain(InRHITextureSwapChain)
+	: RHITexture(AliasedTexture)
+	, RHITextureSwapChain(InRHITextureSwapChain)
 	, SwapChainIndex_RHIThread(0)
 {
 	check(RHITexture);
