Merging //Tasks/UE4/Dev-VR-UEVR-609-LoadingScreenRefactor to Dev-VR-Minimal (//UE4/Dev-VR-Minimal)

diff --git a/Engine/Plugins/Runtime/Oculus/OculusVR/Source/OculusHMD/Private/OculusHMD.cpp b/Engine/Plugins/Runtime/Oculus/OculusVR/Source/OculusHMD/Private/OculusHMD.cpp
--- a/Engine/Plugins/Runtime/Oculus/OculusVR/Source/OculusHMD/Private/OculusHMD.cpp
+++ b/Engine/Plugins/Runtime/Oculus/OculusVR/Source/OculusHMD/Private/OculusHMD.cpp
@@ -539,7 +539,8 @@ namespace OculusHMD
 			// @todo vreditor: If we add support for starting PIE while in VR Editor, we don't want to kill stereo mode when exiting PIE
 			if (Splash->IsShown())
 			{
-				Splash->Hide();
+				Splash->HideLoadingScreen(); // This will only request hiding the screen
+				Splash->UpdateLoadingScreen_GameThread(); // Update is needed to complete removing the loading screen
 			}
 			EnableStereo(false);
 			ReleaseDevice();
@@ -1660,34 +1661,6 @@ namespace OculusHMD
 		}
 	}
 
-
-	void FOculusHMD::UpdateSplashScreen()
-	{
-		if (!GetSplash() || !IsInGameThread())
-		{
-			return;
-		}
-
-		Flags.bNeedSplashUpdate = true;
-	}
-
-	void FOculusHMD::UpdateSplashScreen_GameThread()
-	{
-		if (Flags.bNeedSplashUpdate)
-		{
-			if (bSplashIsShown)
-			{
-				Splash->Show();
-			}
-			else
-			{
-				Splash->Hide();
-			}
-
-			Flags.bNeedSplashUpdate = false;
-		}
-	}
-
 	void FOculusHMD::SetSplashRotationToForward()
 	{
 		//if update splash screen is shown, update the head orientation default to recenter splash screens
@@ -3141,7 +3114,7 @@ namespace OculusHMD
 
 		if (!Frame.IsValid())
 		{
-			UpdateSplashScreen_GameThread(); //the result of this is used in CreateGameFrame to know if Frame is a "real" one or a "splash" one.
+			Splash->UpdateLoadingScreen_GameThread(); //the result of this is used in CreateGameFrame to know if Frame is a "real" one or a "splash" one.
 			Frame = CreateNewGameFrame();
 			NextFrameToRender = Frame;
 
diff --git a/Engine/Plugins/Runtime/Oculus/OculusVR/Source/OculusHMD/Private/OculusHMD.h b/Engine/Plugins/Runtime/Oculus/OculusVR/Source/OculusHMD/Private/OculusHMD.h
--- a/Engine/Plugins/Runtime/Oculus/OculusVR/Source/OculusHMD/Private/OculusHMD.h
+++ b/Engine/Plugins/Runtime/Oculus/OculusVR/Source/OculusHMD/Private/OculusHMD.h
@@ -105,6 +105,7 @@ public:
 	virtual void OnEndPlay(FWorldContext& InWorldContext) override;
 	virtual bool OnStartGameFrame(FWorldContext& WorldContext) override;
 	virtual bool OnEndGameFrame(FWorldContext& WorldContext) override;
+	virtual class IXRLoadingScreen* CreateLoadingScreen() override { return GetSplash(); }
 
 	// IHeadMountedDisplay
 	virtual bool IsHMDConnected() override;
@@ -169,7 +170,7 @@ public:
 	//virtual void UseImplicitHmdPosition(bool bInImplicitHmdPosition) override;
 	//virtual bool GetUseImplicitHmdPosition() override;
 
-	// FHeadMoundedDisplayBase interface
+	// FHeadMountedDisplayBase interface
 	virtual FVector2D GetEyeCenterPoint_RenderThread(EStereoscopicPass StereoPassType) const override;
 	virtual FIntRect GetFullFlatEyeRect_RenderThread(FTexture2DRHIRef EyeTexture) const override;
 	virtual void CopyTexture_RenderThread(FRHICommandListImmediate& RHICmdList, FRHITexture2D* SrcTexture, FIntRect SrcRect, FRHITexture2D* DstTexture, FIntRect DstRect, bool bClearBlack, bool bNoAlpha) const override;
@@ -192,10 +193,11 @@ public:
 	virtual void SetLayerDesc(uint32 LayerId, const IStereoLayers::FLayerDesc& InLayerDesc) override;
 	virtual bool GetLayerDesc(uint32 LayerId, IStereoLayers::FLayerDesc& OutLayerDesc) override;
 	virtual void MarkTextureForUpdate(uint32 LayerId) override;
-	virtual void UpdateSplashScreen() override;
 	virtual IStereoLayers::FLayerDesc GetDebugCanvasLayerDesc(FTextureRHIRef Texture) override;
 	virtual void GetAllocatedTexture(uint32 LayerId, FTextureRHIRef &Texture, FTextureRHIRef &LeftTexture) override;
 	virtual bool ShouldCopyDebugLayersToSpectatorScreen() const override { return true; }
+	virtual void PushLayerState(bool) override { /* Todo */ }
+	virtual void PopLayerState() override { /* Todo */ }
 
 	// ISceneViewExtension
 	virtual void SetupViewFamily(FSceneViewFamily& InViewFamily) override;
@@ -224,7 +226,6 @@ protected:
 	void ApplicationResumeDelegate();
 	void SetupOcclusionMeshes();
 	void UpdateStereoRenderingParams();
-	void UpdateSplashScreen_GameThread();
 	void UpdateHmdRenderInfo();
 	void InitializeEyeLayer_RenderThread(FRHICommandListImmediate& RHICmdList);
 	void ApplySystemOverridesOnStereo(bool force = false);
@@ -365,7 +366,6 @@ protected:
 
 			uint64	bNeedEnableStereo : 1;
 			uint64	bNeedDisableStereo : 1;
-			uint64  bNeedSplashUpdate : 1;
 		};
 		uint64 Raw;
 	} Flags;
diff --git a/Engine/Plugins/Runtime/Oculus/OculusVR/Source/OculusHMD/Private/OculusHMD_Splash.cpp b/Engine/Plugins/Runtime/Oculus/OculusVR/Source/OculusHMD/Private/OculusHMD_Splash.cpp
--- a/Engine/Plugins/Runtime/Oculus/OculusVR/Source/OculusHMD/Private/OculusHMD_Splash.cpp
+++ b/Engine/Plugins/Runtime/Oculus/OculusVR/Source/OculusHMD/Private/OculusHMD_Splash.cpp
@@ -29,6 +29,8 @@ FSplash::FSplash(FOculusHMD* InOculusHMD) :
 	NextLayerId(1),
 	bInitialized(false),
 	bIsShown(false),
+	bNeedSplashUpdate(false),
+	bShouldShowSplash(false),
 	SystemDisplayInterval(1 / 90.0f)
 {
 	// Create empty quad layer for black frame
@@ -125,7 +127,7 @@ void FSplash::LoadSettings()
 
 void FSplash::OnPreLoadMap(const FString&)
 {
-	Show();
+	DoShow();
 }
 
 void FSplash::Startup()
@@ -381,6 +383,21 @@ int FSplash::AddSplash(const FOculusSplashDesc& Desc)
 }
 
 
+void FSplash::AddSplash(const FSplashDesc& Splash)
+{
+	FOculusSplashDesc OculusDesc;
+	OculusDesc.TransformInMeters = Splash.Transform;
+	OculusDesc.QuadSizeInMeters = Splash.QuadSize;
+	OculusDesc.DeltaRotation = Splash.DeltaRotation;
+	OculusDesc.bNoAlphaChannel = Splash.bIgnoreAlpha;
+	OculusDesc.bIsDynamic = Splash.bIsDynamic || Splash.bIsExternal;
+	OculusDesc.TextureOffset = Splash.UVRect.Min;
+	OculusDesc.TextureScale = Splash.UVRect.Max;
+	OculusDesc.LoadedTexture = Splash.Texture;
+
+	AddSplash(OculusDesc);
+}
+
 void FSplash::ClearSplashes()
 {
 	CheckInGameThread();
@@ -420,7 +437,7 @@ IStereoLayers::FLayerDesc FSplash::StereoLayerDescFromOculusSplashDesc(FOculusSp
 	return LayerDesc;
 }
 
-void FSplash::Show()
+void FSplash::DoShow()
 {
 	CheckInGameThread();
 
@@ -518,7 +535,7 @@ void FSplash::Show()
 }
 
 
-void FSplash::Hide()
+void FSplash::DoHide()
 {
 	CheckInGameThread();
 
@@ -528,6 +545,38 @@ void FSplash::Hide()
 	StopTicker();
 }
 
+void FSplash::UpdateLoadingScreen_GameThread()
+{
+	if (bNeedSplashUpdate)
+	{
+		if (bShouldShowSplash)
+		{
+			DoShow();
+		}
+		else
+		{
+			DoHide();
+		}
+
+		bNeedSplashUpdate = false;
+	}
+}
+
+void FSplash::ShowLoadingScreen()
+{
+	bShouldShowSplash = true;
+
+	// DoShow will be called from UpdateSplashScreen_Gamethread().
+	// This can can happen if the splashes are already being shown, as it will reset the relative positions and delta rotations of the layers.
+	bNeedSplashUpdate = true;  
+}
+
+void FSplash::HideLoadingScreen()
+{
+	bShouldShowSplash = false;
+	bNeedSplashUpdate = bIsShown; // no need to call DoHide when the splash is already hidden
+}
+
 void FSplash::UnloadTextures()
 {
 	CheckInGameThread();
diff --git a/Engine/Plugins/Runtime/Oculus/OculusVR/Source/OculusHMD/Private/OculusHMD_Splash.h b/Engine/Plugins/Runtime/Oculus/OculusVR/Source/OculusHMD/Private/OculusHMD_Splash.h
--- a/Engine/Plugins/Runtime/Oculus/OculusVR/Source/OculusHMD/Private/OculusHMD_Splash.h
+++ b/Engine/Plugins/Runtime/Oculus/OculusVR/Source/OculusHMD/Private/OculusHMD_Splash.h
@@ -2,6 +2,7 @@
 
 #pragma once
 #include "OculusHMDPrivate.h"
+#include "IXRLoadingScreen.h"
 
 #if OCULUS_HMD_SUPPORTED_PLATFORMS
 #include "OculusHMD_GameFrame.h"
@@ -33,7 +34,7 @@ public:
 // FSplash
 //-------------------------------------------------------------------------------------------------
 
-class FSplash : public TSharedFromThis<FSplash>
+class FSplash : public IXRLoadingScreen, public TSharedFromThis<FSplash>
 {
 protected:
 	class FTicker : public FTickableObjectRenderThread, public TSharedFromThis<FTicker>
@@ -54,7 +55,6 @@ public:
 	virtual ~FSplash();
 
 	void Tick_RenderThread(float DeltaTime);
-	bool IsShown() const { return bIsShown; }
 
 	void Startup();
 	void LoadSettings();
@@ -64,16 +64,25 @@ public:
 
 	void OnPreLoadMap(const FString&);
 
+	// Called from FOculusHMD
+	void UpdateLoadingScreen_GameThread();
+
+	// Internal extended API
 	int AddSplash(const FOculusSplashDesc&);
-	void ClearSplashes();
 	bool GetSplash(unsigned index, FOculusSplashDesc& OutDesc);
 	void StopTicker();
 	void StartTicker();
 
-	void Show();
-	void Hide();
+	// The standard IXRLoadingScreen interface
+	virtual void ShowLoadingScreen() override;
+	virtual void HideLoadingScreen() override;
+	virtual void ClearSplashes() override;
+	virtual void AddSplash(const FSplashDesc& Splash) override;
+	virtual bool IsShown() const override { return bIsShown; }
 
 protected:
+	void DoShow();
+	void DoHide();
 	void UnloadTextures();
 	void LoadTexture(FSplashLayer& InSplashLayer);
 	void UnloadTexture(FSplashLayer& InSplashLayer);
@@ -101,6 +110,8 @@ protected:
 	// All these flags are only modified from the Game thread
 	bool bInitialized;
 	bool bIsShown;
+	bool bNeedSplashUpdate;
+	bool bShouldShowSplash;
 
 	float SystemDisplayInterval;
 	double LastTimeInSeconds;
diff --git a/Engine/Plugins/Runtime/Steam/SteamVR/Source/SteamVR/Private/SteamVRHMD.h b/Engine/Plugins/Runtime/Steam/SteamVR/Source/SteamVR/Private/SteamVRHMD.h
--- a/Engine/Plugins/Runtime/Steam/SteamVR/Source/SteamVR/Private/SteamVRHMD.h
+++ b/Engine/Plugins/Runtime/Steam/SteamVR/Source/SteamVR/Private/SteamVRHMD.h
@@ -170,7 +170,6 @@ public:
 
 	// IStereoLayers interface
 	// Create/Set/Get/Destroy inherited from TStereoLayerManager
-	virtual void UpdateSplashScreen() override;
 	virtual void GetAllocatedTexture(uint32 LayerId, FTextureRHIRef &Texture, FTextureRHIRef &LeftTexture) override;
 	virtual bool ShouldCopyDebugLayersToSpectatorScreen() const override { return true; }
 
diff --git a/Engine/Plugins/Runtime/Steam/SteamVR/Source/SteamVR/Private/SteamVRRender.cpp b/Engine/Plugins/Runtime/Steam/SteamVR/Source/SteamVR/Private/SteamVRRender.cpp
--- a/Engine/Plugins/Runtime/Steam/SteamVR/Source/SteamVR/Private/SteamVRRender.cpp
+++ b/Engine/Plugins/Runtime/Steam/SteamVR/Source/SteamVR/Private/SteamVRRender.cpp
@@ -34,7 +34,7 @@ void FSteamVRHMD::RenderTexture_RenderThread(FRHICommandListImmediate& RHICmdLis
 {
 	check(IsInRenderingThread());
 
-	if (bSplashIsShown)
+	if (bSplashIsShown || !IsBackgroundLayerVisible())
 	{
 		FRHIRenderPassInfo RPInfo(SrcTexture, ERenderTargetActions::DontLoad_Store);
 		RHICmdList.BeginRenderPass(RPInfo, TEXT("Clear"));
diff --git a/Engine/Plugins/Runtime/Steam/SteamVR/Source/SteamVR/Private/SteamVRStereoLayers.cpp b/Engine/Plugins/Runtime/Steam/SteamVR/Source/SteamVR/Private/SteamVRStereoLayers.cpp
--- a/Engine/Plugins/Runtime/Steam/SteamVR/Source/SteamVR/Private/SteamVRStereoLayers.cpp
+++ b/Engine/Plugins/Runtime/Steam/SteamVR/Source/SteamVR/Private/SteamVRStereoLayers.cpp
@@ -69,52 +69,6 @@ void MarkLayerTextureForUpdate(FSteamVRLayer& Layer)
 	Layer.bUpdateTexture = true;
 }
 
-//=============================================================================
-void FSteamVRHMD::UpdateSplashScreen()
-{
-	FTexture2DRHIRef Texture = (bSplashShowMovie && SplashMovie.IsValid()) ? SplashMovie : SplashTexture;
-	if (bSplashIsShown && Texture.IsValid())
-	{
-		FLayerDesc LayerDesc;
-		LayerDesc.Flags = ELayerFlags::LAYER_FLAG_TEX_NO_ALPHA_CHANNEL;
-		LayerDesc.PositionType = ELayerType::TrackerLocked;
-		LayerDesc.Texture = Texture;
-		const FIntPoint TextureSize = Texture->GetSizeXY();
-		const float	InvAspectRatio = (TextureSize.X > 0) ? float(TextureSize.Y) / float(TextureSize.X) : 1.0f;
-
-		// Get the current pose of the HMD
-		FVector HMDPosition;
-		FQuat HMDOrientation;
-		GetCurrentPose(IXRTrackingSystem::HMDDeviceId, HMDOrientation, HMDPosition);
-
-		FTransform Translation(FVector(500.0f, 0.0f, 100.0f) + SplashOffset);
-		FRotator Rotation(HMDOrientation);
-		Rotation.Pitch = 0.0f;
-		Rotation.Roll = 0.0f;
-		LayerDesc.Transform = Translation * FTransform(Rotation.Quaternion());
-
-		// Set texture size to 8m wide, keeping the aspect ratio.
-		LayerDesc.QuadSize = FVector2D(800.0f, 800.0f*InvAspectRatio) * SplashScale;
-
-		if (SplashLayerHandle)
-		{
-			SetLayerDesc(SplashLayerHandle, LayerDesc);
-		}
-		else
-		{
-			SplashLayerHandle = CreateLayer(LayerDesc);
-		}
-	}
-	else
-	{
-		if (SplashLayerHandle)
-		{
-			DestroyLayer(SplashLayerHandle);
-			SplashLayerHandle = 0;
-		}
-	}
-}
-
 //=============================================================================
 void FSteamVRHMD::UpdateLayer(struct FSteamVRLayer& Layer, uint32 LayerId, bool bIsValid) const
 {
diff --git a/Engine/Source/Runtime/Engine/Classes/Kismet/StereoLayerFunctionLibrary.h b/Engine/Source/Runtime/Engine/Classes/Kismet/StereoLayerFunctionLibrary.h
--- a/Engine/Source/Runtime/Engine/Classes/Kismet/StereoLayerFunctionLibrary.h
+++ b/Engine/Source/Runtime/Engine/Classes/Kismet/StereoLayerFunctionLibrary.h
@@ -24,19 +24,19 @@ public:
 	* @param Offset				(in) Position in UE Units to offset the Splash Screen by
 	* @param ShowLoadingMovie	(in) Whether the splash screen presents loading movies.
 	*/
-	UFUNCTION(BlueprintCallable, Category = "VR")
+	UFUNCTION(BlueprintCallable, Category = "VR", meta=(DeprecatedFunction, DeprecationMessage="Please use Add Loading Screen Splash instead."))
 		static void SetSplashScreen(class UTexture* Texture, FVector2D Scale = FVector2D(1.0f, 1.0f), FVector Offset = FVector(0.0f, 0.0f, 0.0f), bool bShowLoadingMovie = false, bool bShowOnSet = false);
 
 	/**
 	* Show the splash screen and override the VR display
 	*/
-	UFUNCTION(BlueprintCallable, Category = "VR")
+	UFUNCTION(BlueprintCallable, Category = "VR", meta = (DeprecatedFunction, DeprecationMessage = "Please use Show Loading Screen instead."))
 	static void ShowSplashScreen();
 
 	/**
 	* Hide the splash screen and return to normal display.
 	*/
-	UFUNCTION(BlueprintCallable, Category = "VR")
+	UFUNCTION(BlueprintCallable, Category = "VR", meta = (DeprecatedFunction, DeprecationMessage = "Please use Hide Loading Screen instead."))
 	static void HideSplashScreen();
 
 	/**
diff --git a/Engine/Source/Runtime/Engine/Private/StereoLayerFunctionLibrary.cpp b/Engine/Source/Runtime/Engine/Private/StereoLayerFunctionLibrary.cpp
--- a/Engine/Source/Runtime/Engine/Private/StereoLayerFunctionLibrary.cpp
+++ b/Engine/Source/Runtime/Engine/Private/StereoLayerFunctionLibrary.cpp
@@ -7,6 +7,18 @@
 #include "Engine/Engine.h"
 #include "IStereoLayers.h"
 #include "StereoRendering.h"
+#include "IXRTrackingSystem.h"
+#include "IXRLoadingScreen.h"
+
+static IXRLoadingScreen* GetLoadingScreen()
+{
+	if (GEngine && GEngine->XRSystem.IsValid())
+	{
+		return GEngine->XRSystem->GetLoadingScreen();
+	}
+
+	return nullptr;
+}
 
 static IStereoLayers* GetStereoLayers()
 {
@@ -18,6 +30,7 @@ static IStereoLayers* GetStereoLayers()
 	return nullptr;
 }
 
+
 class FAutoShow: public TSharedFromThis<FAutoShow>
 {
 public:
@@ -30,19 +43,19 @@ public:
 
 void FAutoShow::OnPreLoadMap(const FString&)
 {
-	IStereoLayers* StereoLayers = GetStereoLayers();
-	if (StereoLayers)
+	IXRLoadingScreen* LoadingScreen = GetLoadingScreen();
+	if (LoadingScreen)
 	{
-		StereoLayers->ShowSplashScreen();
+		LoadingScreen->ShowLoadingScreen();
 	}
 }
 
 void FAutoShow::OnPostLoadMap(UWorld* LoadedWorld)
 {
-	IStereoLayers* StereoLayers = GetStereoLayers();
-	if (StereoLayers)
+	IXRLoadingScreen* LoadingScreen = GetLoadingScreen();
+	if (LoadingScreen)
 	{
-		StereoLayers->HideSplashScreen();
+		LoadingScreen->HideLoadingScreen();
 	}
 }
 
@@ -67,32 +80,38 @@ UStereoLayerFunctionLibrary::UStereoLayerFunctionLibrary(const FObjectInitialize
 
 void UStereoLayerFunctionLibrary::SetSplashScreen(class UTexture* Texture, FVector2D Scale, FVector Offset, bool bShowLoadingMovie, bool bShowOnSet)
 {
-	IStereoLayers* StereoLayers = GetStereoLayers();
-	if (StereoLayers && Texture && Texture->Resource)
+	IXRLoadingScreen* LoadingScreen = GetLoadingScreen();
+	if (LoadingScreen && Texture && Texture->Resource)
 	{
-		StereoLayers->SetSplashScreen(Texture->Resource->TextureRHI, Scale, Offset, bShowLoadingMovie);
+		LoadingScreen->ClearSplashes();
+		IXRLoadingScreen::FSplashDesc Splash;
+		Splash.Texture = Texture->Resource->TextureRHI;
+		Splash.QuadSize = Scale;
+		Splash.Transform = FTransform(Offset);
+		LoadingScreen->AddSplash(Splash);
+
 		if (bShowOnSet)
 		{
-			StereoLayers->ShowSplashScreen();
+			LoadingScreen->ShowLoadingScreen();
 		}
 	}
 }
 
 void UStereoLayerFunctionLibrary::ShowSplashScreen()
 {
-	IStereoLayers* StereoLayers = GetStereoLayers();
-	if (StereoLayers)
+	IXRLoadingScreen* LoadingScreen = GetLoadingScreen();
+	if (LoadingScreen)
 	{
-		StereoLayers->ShowSplashScreen();
+		LoadingScreen->ShowLoadingScreen();
 	}
 }
 
 void UStereoLayerFunctionLibrary::HideSplashScreen()
 {
-	IStereoLayers* StereoLayers = GetStereoLayers();
-	if (StereoLayers)
+	IXRLoadingScreen* LoadingScreen = GetLoadingScreen();
+	if (LoadingScreen)
 	{
-		StereoLayers->HideSplashScreen();
+		LoadingScreen->HideLoadingScreen();
 	}
 }
 
diff --git a/Engine/Source/Runtime/Engine/Public/IStereoLayers.h b/Engine/Source/Runtime/Engine/Public/IStereoLayers.h
--- a/Engine/Source/Runtime/Engine/Public/IStereoLayers.h
+++ b/Engine/Source/Runtime/Engine/Public/IStereoLayers.h
@@ -50,7 +50,8 @@ public:
 		void SetLayerId(uint32 InId) { Id = InId; }
 		uint32 GetLayerId() const { return Id; }
 
-		const static uint32	INVALID_LAYER_ID = UINT_MAX;
+		// Layer IDs must be larger than 0
+		const static uint32	INVALID_LAYER_ID = 0; 
 		// The layer's ID
 		uint32				Id			= INVALID_LAYER_ID;
 		// View space transform
@@ -81,7 +82,7 @@ public:
 	virtual ~IStereoLayers() { }
 
 	/**
-	 * Creates a new layer from a given texture resource, which is projected on top of the world as a quad
+	 * Creates a new layer from a given texture resource, which is projected on top of the world as a quad.
 	 *
 	 * @param	InLayerDesc		A reference to the texture resource to be used on the quad
 	 * @return	A unique identifier for the layer created
@@ -89,12 +90,65 @@ public:
 	virtual uint32 CreateLayer(const FLayerDesc& InLayerDesc) = 0;
 	
 	/**
-	 * Destroys the specified layer, stopping it from rendering over the world
+	 * Destroys the specified layer, stopping it from rendering over the world.
 	 *
 	 * @param	LayerId		The ID of layer to be destroyed
 	 */
 	virtual void DestroyLayer(uint32 LayerId) = 0;
 
+	/**
+	 * Saves the current stereo layer state on a stack to later restore them.
+	 *
+	 * Useful for creating temporary overlays that should be torn down later.
+	 *
+	 * When bPreserve is false, existing layers will be temporarily disabled and restored again when calling PopLayerState()
+	 * The disabled layer's properties are still accessible by calling Get and SetLayerDesc, but nothing will change until after
+	 * the state has been restored. Calling DestroyLayer on an inactive layer, will prevent it from being restored when PopLayerState() is called.
+	 *
+	 * When bPreserve is true, existing layers will remain active, but when calling PopLayerState(), any changed properties
+	 * will be restored back to their previous values. Calling DestroyLayer on an active layer id will make the layer inactive. The layer
+	 * will be reactivated when the state is restored. (You can call DestroyLayer multiple times on the same layer id to remove successively older
+	 * versions of a layer.)
+	 *
+	 * In either case, layers created after PushLayerState() will be destroyed upon calling PopLayerState(). 
+	 * 
+	 * @param	bPreserve	Whether the existing layers should be preserved after saving the state. If false all existing layers will be disabled.
+	 */
+	virtual void PushLayerState(bool bPreserve = false) {};
+
+	/**
+	 * Restores the stereo layer state from the last save state. 
+	 * 
+	 * Currently active layers will be destroyed and replaced with the previous state.
+	 */
+	virtual void PopLayerState() {};
+
+	/**
+	 * Returns true if the StereoLayers implementation supports saving and restoring state using Push/PopLayerState()
+	 */
+	virtual bool SupportsLayerState() { return false; }
+
+	/** 
+	 * Optional method to hide the 3D scene and only render the stereo overlays. 
+	 * No-op if not supported by the platform.
+	 *
+	 * If pushing and popping layer state is supported, the visibility of the background layer should be part of
+	 * the saved state.
+	 */
+	virtual void HideBackgroundLayer() {}
+
+	/**
+	 * Optional method to undo the effect of hiding the 3D scene.
+	 * No-op if not supported by the platform.
+	 */
+	virtual void ShowBackgroundLayer() {}
+
+	/** 
+	 * Tell if the background layer is visible. Platforms that do not implement Hide/ShowBackgroundLayer() 
+	 * always return true.
+	 */
+	virtual bool IsBackgroundLayerVisible() const { return true; }
+
 	/**
 	 * Set the a new layer description
 	 *
@@ -122,7 +176,7 @@ public:
 	/**
 	 * Update splash screens from current state
 	 */
-	virtual void UpdateSplashScreen() = 0;
+	virtual void UpdateSplashScreen() {};
 
 	/**
 	* If true the debug layers are copied to the spectator screen, because they do not naturally end up on the spectator screen as part of the 3d view.
@@ -138,6 +192,7 @@ public:
 	* @param Offset				(in) Position from which to start rendering the texture.
 	* @param ShowLoadingMovie	(in) Whether the splash screen presents loading movies.
 	*/
+	UE_DEPRECATED(4.24, "Use the IXRLoadingScreen interface instead of IStereoLayers::*SplashScreen")
 	void SetSplashScreen(FTextureRHIRef Texture, FVector2D Scale, FVector Offset, bool bShowLoadingMovie)
 	{
 		bSplashShowMovie = bShowLoadingMovie;
@@ -153,6 +208,7 @@ public:
 	/**
 	* Show the splash screen and override the normal VR display
 	*/
+	UE_DEPRECATED(4.24, "Use the IXRLoadingScreen interface instead of IStereoLayers::*SplashScreen")
 	void ShowSplashScreen()
 	{
 		bSplashIsShown = true;
@@ -162,6 +218,7 @@ public:
 	/**
 	* Hide the splash screen and return to normal display.
 	*/
+	UE_DEPRECATED(4.24, "Use the IXRLoadingScreen interface instead of IStereoLayers::*SplashScreen")
 	void HideSplashScreen()
 	{
 		bSplashIsShown = false;
@@ -173,6 +230,7 @@ public:
 	*
 	* @param InMovieTexture		(in) A movie texture to be used for the splash. B8R8G8A8 format.
 	*/
+	UE_DEPRECATED(4.24, "Use the IXRLoadingScreen interface instead of IStereoLayers::*SplashScreen")
 	void SetSplashScreenMovie(FTextureRHIRef Texture)
 	{
 		SplashMovie = nullptr;
diff --git a/Engine/Source/Runtime/HeadMountedDisplay/Private/DefaultStereoLayers.cpp b/Engine/Source/Runtime/HeadMountedDisplay/Private/DefaultStereoLayers.cpp
--- a/Engine/Source/Runtime/HeadMountedDisplay/Private/DefaultStereoLayers.cpp
+++ b/Engine/Source/Runtime/HeadMountedDisplay/Private/DefaultStereoLayers.cpp
@@ -14,6 +14,7 @@
 #include "ClearQuad.h"
 #include "SceneView.h"
 #include "CommonRenderResources.h"
+#include "IXRLoadingScreen.h"
 
 namespace 
 {
@@ -250,8 +251,13 @@ void FDefaultStereoLayers::PostRenderView_RenderThread(FRHICommandListImmediate&
 
 	FRHIRenderPassInfo RPInfo(RenderTarget, ERenderTargetActions::Load_Store);
 	RHICmdList.BeginRenderPass(RPInfo, TEXT("StereoLayerRender"));
-
 	RHICmdList.SetViewport(RenderParams.Viewport.Min.X, RenderParams.Viewport.Min.Y, 0, RenderParams.Viewport.Max.X, RenderParams.Viewport.Max.Y, 1.0f);
+
+	if (bSplashIsShown || !IsBackgroundLayerVisible())
+	{
+		DrawClearQuad(RHICmdList, FLinearColor::Black);
+	}
+
 	StereoLayerRender(RHICmdList, SortedSceneLayers, RenderParams);
 	
 	// Optionally render face-locked layers into a non-reprojected target if supported by the HMD platform
@@ -277,47 +283,6 @@ bool FDefaultStereoLayers::IsActiveThisFrame(class FViewport* InViewport) const
 	return GEngine && GEngine->IsStereoscopic3D(InViewport);
 }
 
-void FDefaultStereoLayers::UpdateSplashScreen()
-{
-	FTexture2DRHIRef Texture = (bSplashShowMovie && SplashMovie.IsValid()) ? SplashMovie : SplashTexture;
-	if (bSplashIsShown && Texture.IsValid())
-	{
-		FQuat Orientation;
-		FVector Position;
-
-		HMDDevice->GetCurrentPose(IXRTrackingSystem::HMDDeviceId, Orientation, Position);
-		FLayerDesc LayerDesc;
-		LayerDesc.Flags = ELayerFlags::LAYER_FLAG_TEX_NO_ALPHA_CHANNEL;
-		LayerDesc.PositionType = ELayerType::TrackerLocked;
-		LayerDesc.Texture = Texture;
-
-		FTransform Translation(FVector(500.0f, 0.0f, 100.0f) + SplashOffset);
-		FRotator Rotation(Orientation);
-		Rotation.Pitch = 0.0f;
-		Rotation.Roll = 0.0f;
-		LayerDesc.Transform = Translation * FTransform(Rotation.Quaternion());
-
-		LayerDesc.QuadSize = FVector2D(800.0f, 450.0f) * SplashScale;
-
-		if (SplashLayerHandle)
-		{
-			SetLayerDesc(SplashLayerHandle, LayerDesc);
-		}
-		else
-		{
-			SplashLayerHandle = CreateLayer(LayerDesc);
-		}
-	}
-	else
-	{
-		if (SplashLayerHandle)
-		{
-			DestroyLayer(SplashLayerHandle);
-			SplashLayerHandle = 0;
-		}
-	}
-}
-
 void FDefaultStereoLayers::SetupViewFamily(FSceneViewFamily& InViewFamily)
 {
 	// Initialize HMD position.
diff --git a/Engine/Source/Runtime/HeadMountedDisplay/Private/DefaultXRLoadingScreen.cpp b/Engine/Source/Runtime/HeadMountedDisplay/Private/DefaultXRLoadingScreen.cpp
new file mode 100644
--- /dev/null
+++ b/Engine/Source/Runtime/HeadMountedDisplay/Private/DefaultXRLoadingScreen.cpp
@@ -0,0 +1,139 @@
+// Copyright 1998-2019 Epic Games, Inc. All Rights Reserved.
+
+#include "DefaultXRLoadingScreen.h"
+#include "IXRTrackingSystem.h"
+#include "XRThreadUtils.h"
+
+using FLayerDesc = IStereoLayers::FLayerDesc;
+using FSplashDesc = IXRLoadingScreen::FSplashDesc;
+
+
+FDefaultXRLoadingScreen::FDefaultXRLoadingScreen(class IXRTrackingSystem* InTrackingSystem)
+	: TXRLoadingScreenBase(InTrackingSystem)
+{
+	check(TrackingSystem->GetStereoRenderingDevice().IsValid());
+}
+
+IStereoLayers* FDefaultXRLoadingScreen::GetStereoLayers() const
+{
+	IStereoLayers* StereoLayers = nullptr;
+	auto StereoRendering = TrackingSystem->GetStereoRenderingDevice();
+	if (StereoRendering.IsValid())
+	{
+		StereoLayers = StereoRendering->GetStereoLayers();
+		if (!StereoLayers || !StereoLayers->SupportsLayerState())
+		{
+			StereoLayers = nullptr;
+			UE_LOG(LogHMD, Warning, TEXT("FDefaultXRLoadingScreen requires a working IStereoLayers implementation supporting IStereoLayers::Push/PopLayerState."));
+		}
+	}
+	return StereoLayers;
+}
+
+void FDefaultXRLoadingScreen::ShowLoadingScreen()
+{
+	auto StereoLayers = GetStereoLayers();
+	if (!StereoLayers)
+	{
+		return;
+	}
+
+	if (!bShowing)
+	{
+		StereoLayers->PushLayerState();
+		StereoLayers->HideBackgroundLayer();
+	}
+	TXRLoadingScreenBase::ShowLoadingScreen();
+}
+
+void FDefaultXRLoadingScreen::HideLoadingScreen()
+{
+	if (!bShowing)
+	{
+		return;
+	}
+	TXRLoadingScreenBase::HideLoadingScreen();
+
+	auto StereoLayers = GetStereoLayers();
+	if (StereoLayers) // StereoLayers may be null if we are shutting down
+	{
+		StereoLayers->PopLayerState();
+	}
+}
+
+void FDefaultXRLoadingScreen::DoHideSplash(FSplashData& Splash)
+{
+	// Pop layer state will destroy all created layers, so we clear the layer id stored locally
+	Splash.LayerId = FLayerDesc::INVALID_LAYER_ID;
+}
+
+void FDefaultXRLoadingScreen::DoShowSplash(FSplashData& Splash)
+{
+	FLayerDesc LayerDesc;
+	bool bUpdate = (Splash.LayerId != FLayerDesc::INVALID_LAYER_ID);
+	auto StereoLayers = GetStereoLayers();
+	float WorldScale = TrackingSystem->GetWorldToMetersScale();
+
+	if (bUpdate)
+	{
+		StereoLayers->GetLayerDesc(Splash.LayerId, LayerDesc);
+	}
+
+	LayerDesc.Transform = Splash.Desc.Transform;
+	LayerDesc.Transform.ScaleTranslation(WorldScale);
+	LayerDesc.Transform *= FTransform(HMDOrientation);
+	LayerDesc.QuadSize = Splash.Desc.QuadSize * WorldScale;
+	LayerDesc.UVRect = Splash.Desc.UVRect;
+	// Sort layers by distance from the viewer
+	LayerDesc.Priority = INT32_MAX - static_cast<int32>(Splash.Desc.Transform.GetTranslation().X * 1000.f);
+	LayerDesc.PositionType = IStereoLayers::TrackerLocked;
+	LayerDesc.ShapeType = Splash.Desc.Texture->GetTextureCube() != nullptr ? IStereoLayers::CubemapLayer : IStereoLayers::QuadLayer;
+	LayerDesc.Texture = Splash.Desc.Texture;
+	LayerDesc.LeftTexture = Splash.Desc.LeftTexture;
+	LayerDesc.Flags = IStereoLayers::LAYER_FLAG_QUAD_PRESERVE_TEX_RATIO
+		| (Splash.Desc.bIgnoreAlpha ? IStereoLayers::LAYER_FLAG_TEX_NO_ALPHA_CHANNEL  : 0)
+		| (Splash.Desc.bIsDynamic   ? IStereoLayers::LAYER_FLAG_TEX_CONTINUOUS_UPDATE : 0)
+		| (Splash.Desc.bIsExternal  ? IStereoLayers::LAYER_FLAG_TEX_EXTERNAL		  : 0);
+
+	if (bUpdate)
+	{
+		StereoLayers->SetLayerDesc(Splash.LayerId, LayerDesc);
+	}
+	else
+	{
+		Splash.LayerId = StereoLayers->CreateLayer(LayerDesc);
+	}
+}
+
+void FDefaultXRLoadingScreen::DoDeleteSplash(FSplashData& Splash)
+{
+	auto StereoLayers = GetStereoLayers();
+	check(StereoLayers); // bShowing will never get set unless stereo layers are available
+	if (Splash.LayerId != FLayerDesc::INVALID_LAYER_ID)
+	{
+		StereoLayers->DestroyLayer(Splash.LayerId);
+		Splash.LayerId = FLayerDesc::INVALID_LAYER_ID;
+	}
+}
+
+void FDefaultXRLoadingScreen::ApplyDeltaRotation(const FSplashData& Splash)
+{
+	if (!bShowing || Splash.LayerId == FLayerDesc::INVALID_LAYER_ID || Splash.Desc.DeltaRotation == FQuat::Identity)
+	{
+		return;
+	}
+	auto StereoLayers = GetStereoLayers();
+
+	FLayerDesc LayerDesc;
+	StereoLayers->GetLayerDesc(Splash.LayerId, LayerDesc);
+	LayerDesc.Transform.SetRotation(LayerDesc.Transform.GetRotation() * Splash.Desc.DeltaRotation);
+	LayerDesc.Transform.NormalizeRotation();
+	StereoLayers->SetLayerDesc(Splash.LayerId, LayerDesc);
+}
+
+FSplashData::FSplashData(const FSplashDesc& InDesc) 
+	: Desc(InDesc)
+	, LayerId(FLayerDesc::INVALID_LAYER_ID)
+{
+
+}
diff --git a/Engine/Source/Runtime/HeadMountedDisplay/Private/IXRLoadingScreen.cpp b/Engine/Source/Runtime/HeadMountedDisplay/Private/IXRLoadingScreen.cpp
new file mode 100644
--- /dev/null
+++ b/Engine/Source/Runtime/HeadMountedDisplay/Private/IXRLoadingScreen.cpp
@@ -0,0 +1,41 @@
+// Copyright 1998-2019 Epic Games, Inc. All Rights Reserved.
+
+#include "IXRLoadingScreen.h"
+#include "IXRTrackingSystem.h"
+#include "EngineGlobals.h"
+#include "Engine/Engine.h"
+
+
+void IXRLoadingScreen::ShowLoadingScreen_Compat(bool bShow, FTexture2DRHIRef Texture, const FVector& Offset, const FVector2D& Scale)
+{
+	// Backwards compatibility with IStereoLayers::ShowSplashScreen
+	IXRLoadingScreen* LoadingScreen = GEngine && GEngine->XRSystem.IsValid()? GEngine->XRSystem->GetLoadingScreen() : nullptr;
+	if (LoadingScreen)
+	{
+		if (bShow)
+		{
+			// Any texture passed in will override all splashes set through the IXRLoadingScreen interface,
+			// otherwise this function only calls Show/HideLoadingScreen.
+			if (Texture.IsValid())
+			{
+				IXRLoadingScreen::FSplashDesc Splash;
+				const FIntPoint TextureSize = Texture->GetSizeXY();
+				const float	InvAspectRatio = (TextureSize.X > 0) ? float(TextureSize.Y) / float(TextureSize.X) : 1.0f;
+
+				Splash.Texture = Texture;
+				Splash.bIgnoreAlpha = true;
+				Splash.Transform = FTransform(FVector(5.0f, 0.0f, 1.0f) + Offset * 0.01f);
+				Splash.QuadSize = FVector2D(8.0f, 8.0f*InvAspectRatio) * Scale;
+
+				LoadingScreen->ClearSplashes();
+				LoadingScreen->AddSplash(Splash);
+			}
+			LoadingScreen->ShowLoadingScreen();
+		}
+		else
+		{
+			LoadingScreen->HideLoadingScreen();
+		}
+	}
+
+}
diff --git a/Engine/Source/Runtime/HeadMountedDisplay/Private/XRLoadingScreenFunctionLibrary.cpp b/Engine/Source/Runtime/HeadMountedDisplay/Private/XRLoadingScreenFunctionLibrary.cpp
new file mode 100644
--- /dev/null
+++ b/Engine/Source/Runtime/HeadMountedDisplay/Private/XRLoadingScreenFunctionLibrary.cpp
@@ -0,0 +1,92 @@
+// Copyright 1998-2019 Epic Games, Inc. All Rights Reserved.
+
+#include "XRLoadingScreenFunctionLibrary.h"
+#include "EngineGlobals.h"
+#include "TextureResource.h"
+#include "Engine/Texture.h"
+#include "Engine/Engine.h"
+#include "IXRTrackingSystem.h"
+
+static IXRLoadingScreen* GetLoadingScreen()
+{
+	if (GEngine && GEngine->XRSystem.IsValid())
+	{
+		return GEngine->XRSystem->GetLoadingScreen();
+	}
+	
+	return nullptr;
+}
+
+UXRLoadingScreenFunctionLibrary::UXRLoadingScreenFunctionLibrary(const FObjectInitializer& ObjectInitializer)
+	: Super(ObjectInitializer)
+{
+}
+
+void UXRLoadingScreenFunctionLibrary::SetLoadingScreen(class UTexture* Texture, FVector2D Scale, FVector Offset, bool bShowLoadingMovie, bool bShowOnSet)
+{
+	IXRLoadingScreen* LoadingSrcreen = GetLoadingScreen();
+	if (LoadingSrcreen && Texture && Texture->Resource)
+	{
+		LoadingSrcreen->ClearSplashes();
+		const bool bIsExternal = Texture->GetMaterialType() == MCT_TextureExternal;
+		IXRLoadingScreen::FSplashDesc Splash;
+		Splash.Transform = FTransform(Offset);
+		Splash.QuadSize = Scale;
+		Splash.bIsDynamic = bShowLoadingMovie || bIsExternal;
+		Splash.bIsExternal = bIsExternal;
+		Splash.Texture = Texture->Resource->TextureRHI;
+		LoadingSrcreen->AddSplash(Splash);
+
+		if (bShowOnSet)
+		{
+			LoadingSrcreen->ShowLoadingScreen();
+		}
+	}
+}
+
+void UXRLoadingScreenFunctionLibrary::ClearLoadingScreenSplashes()
+{
+	IXRLoadingScreen* LoadingSrcreen = GetLoadingScreen();
+	if (LoadingSrcreen)
+	{
+		LoadingSrcreen->ClearSplashes();
+	}
+
+}
+
+void UXRLoadingScreenFunctionLibrary::AddLoadingScreenSplash(class UTexture* Texture, FVector Translation, FRotator Rotation, FVector2D Size, FRotator DeltaRotation, bool bClearBeforeAdd)
+{
+	IXRLoadingScreen* LoadingSrcreen = GetLoadingScreen();
+	if (LoadingSrcreen && Texture && Texture->Resource)
+	{
+		if (bClearBeforeAdd)
+		{
+			LoadingSrcreen->ClearSplashes();
+		}
+
+		IXRLoadingScreen::FSplashDesc Splash;
+		Splash.Texture = Texture->Resource->TextureRHI;
+		Splash.QuadSize = Size;
+		Splash.Transform = FTransform(Rotation, Translation);
+		Splash.DeltaRotation = FQuat(DeltaRotation);
+		LoadingSrcreen->AddSplash(Splash);
+	}
+}
+
+void UXRLoadingScreenFunctionLibrary::ShowLoadingScreen()
+{
+	IXRLoadingScreen* LoadingSrcreen = GetLoadingScreen();
+	if (LoadingSrcreen)
+	{
+		LoadingSrcreen->ShowLoadingScreen();
+	}
+}
+
+void UXRLoadingScreenFunctionLibrary::HideLoadingScreen()
+{
+	IXRLoadingScreen* LoadingSrcreen = GetLoadingScreen();
+	if (LoadingSrcreen)
+	{
+		LoadingSrcreen->HideLoadingScreen();
+	}
+}
diff --git a/Engine/Source/Runtime/HeadMountedDisplay/Private/XRTrackingSystemBase.cpp b/Engine/Source/Runtime/HeadMountedDisplay/Private/XRTrackingSystemBase.cpp
--- a/Engine/Source/Runtime/HeadMountedDisplay/Private/XRTrackingSystemBase.cpp
+++ b/Engine/Source/Runtime/HeadMountedDisplay/Private/XRTrackingSystemBase.cpp
@@ -6,6 +6,8 @@
 #include "Engine/LocalPlayer.h"
 #include "Camera/CameraComponent.h"
 #include "GameFramework/PlayerController.h"
+#include "DefaultXRLoadingScreen.h"
+
 
 // Tracking system delegates
 FXRTrackingSystemDelegates::FXRTrackingOriginChanged FXRTrackingSystemDelegates::OnXRTrackingOriginChanged;
@@ -14,6 +16,7 @@ FXRTrackingSystemDelegates::FXRTrackingOriginChanged FXRTrackingSystemDelegates:
 
 
 FXRTrackingSystemBase::FXRTrackingSystemBase(IARSystemSupport* InARImplementation)
+	: LoadingScreen(nullptr)
 {
 	ARCompositionComponent = MakeShared<FARSupportInterface , ESPMode::ThreadSafe>(InARImplementation, this);
 }
@@ -24,6 +27,13 @@ FXRTrackingSystemBase::~FXRTrackingSystemBase()
 	{
 		ARCompositionComponent.Reset();
 	}
+
+	if (LoadingScreen)
+	{
+		LoadingScreen->HideLoadingScreen();
+		delete LoadingScreen;
+		LoadingScreen = nullptr;
+	}
 }
 
 uint32 FXRTrackingSystemBase::CountTrackedDevices(EXRTrackedDeviceType Type /*= EXRTrackedDeviceType::Any*/)
@@ -111,6 +121,25 @@ void FXRTrackingSystemBase::UpdateExternalTrackingPosition(const FTransform& Ext
 }
 
 
+class IXRLoadingScreen* FXRTrackingSystemBase::GetLoadingScreen()
+{
+	if (!LoadingScreen)
+	{
+		LoadingScreen = CreateLoadingScreen();
+	}
+	return LoadingScreen;
+}
+
+class IXRLoadingScreen* FXRTrackingSystemBase::CreateLoadingScreen()
+{
+	if (GetStereoRenderingDevice().IsValid() && GetStereoRenderingDevice()->GetStereoLayers() && GetStereoRenderingDevice()->GetStereoLayers()->SupportsLayerState())
+	{
+		return new FDefaultXRLoadingScreen(this);
+	}
+	return nullptr;
+}
+
+
 TSharedPtr<FARSupportInterface , ESPMode::ThreadSafe> FXRTrackingSystemBase::GetARCompositionComponent()
 {
 	return ARCompositionComponent;
@@ -121,7 +150,6 @@ const TSharedPtr<const FARSupportInterface , ESPMode::ThreadSafe> FXRTrackingSys
 	return ARCompositionComponent;
 }
 
-
 FTransform FXRTrackingSystemBase::RefreshTrackingToWorldTransform(FWorldContext& WorldContext)
 {
 	CachedTrackingToWorld = ComputeTrackingToWorldTransform(WorldContext);
diff --git a/Engine/Source/Runtime/HeadMountedDisplay/Public/DefaultStereoLayers.h b/Engine/Source/Runtime/HeadMountedDisplay/Public/DefaultStereoLayers.h
--- a/Engine/Source/Runtime/HeadMountedDisplay/Public/DefaultStereoLayers.h
+++ b/Engine/Source/Runtime/HeadMountedDisplay/Public/DefaultStereoLayers.h
@@ -15,8 +15,6 @@ class HEADMOUNTEDDISPLAY_API FDefaultStereoLayers : public FSimpleLayerManager,
 public:
 	FDefaultStereoLayers(const FAutoRegister& AutoRegister, FHeadMountedDisplayBase* InHMDDevice);
 
-	virtual void UpdateSplashScreen() override;
-
 	/** ISceneViewExtension interface */
 	virtual void SetupViewFamily(FSceneViewFamily& InViewFamily) override;
 	virtual void SetupView(FSceneViewFamily& InViewFamily, FSceneView& InView) override {}
diff --git a/Engine/Source/Runtime/HeadMountedDisplay/Public/DefaultXRLoadingScreen.h b/Engine/Source/Runtime/HeadMountedDisplay/Public/DefaultXRLoadingScreen.h
new file mode 100644
--- /dev/null
+++ b/Engine/Source/Runtime/HeadMountedDisplay/Public/DefaultXRLoadingScreen.h
@@ -0,0 +1,40 @@
+// Copyright 1998-2019 Epic Games, Inc. All Rights Reserved.
+
+#pragma once
+
+#include "XRLoadingScreenBase.h"
+#include "IStereoLayers.h"
+
+/**
+ * Default Loading Screen implementation based on the IStereoLayer interface.
+ * It requires an XR tracking system with stereo rendering and stereo layers support.
+ */
+
+struct FSplashData {
+	IXRLoadingScreen::FSplashDesc	Desc;
+	uint32		LayerId;
+
+	FSplashData(const IXRLoadingScreen::FSplashDesc& InDesc);
+};
+
+class HEADMOUNTEDDISPLAY_API FDefaultXRLoadingScreen : public TXRLoadingScreenBase<FSplashData>
+{
+public:
+	FDefaultXRLoadingScreen(class IXRTrackingSystem* InTrackingSystem);
+
+	/* IXRLoadingScreen interface */
+	virtual void ShowLoadingScreen() override;
+	virtual void HideLoadingScreen() override;
+
+private:
+
+	IStereoLayers* GetStereoLayers() const;
+
+protected:
+	virtual void DoShowSplash(FSplashData& Splash) override;
+	virtual void DoHideSplash(FSplashData& Splash) override;
+	virtual void DoAddSplash(FSplashData& Splash) override {}
+	virtual void DoDeleteSplash(FSplashData& Splash) override;
+	virtual void ApplyDeltaRotation(const FSplashData& Splash) override;
+
+};
\ No newline at end of file
diff --git a/Engine/Source/Runtime/HeadMountedDisplay/Public/IXRLoadingScreen.h b/Engine/Source/Runtime/HeadMountedDisplay/Public/IXRLoadingScreen.h
new file mode 100644
--- /dev/null
+++ b/Engine/Source/Runtime/HeadMountedDisplay/Public/IXRLoadingScreen.h
@@ -0,0 +1,90 @@
+// Copyright 1998-2019 Epic Games, Inc. All Rights Reserved.
+
+#pragma once
+
+#include "CoreMinimal.h"
+#include "HeadMountedDisplayTypes.h"
+
+
+
+/**
+ * Interface used to register and show loading splash screen layers.
+ * A loading screen consists of one or more splash layers shown instead of the current scene during level load.
+ *
+ * To reduce duplicated code, implementations should use the FXRLoadingScreenBase of FDefaultXRLoadingScreen instead of implementing this interface directly.
+ */
+class HEADMOUNTEDDISPLAY_API  IXRLoadingScreen 
+{
+public:
+
+	virtual ~IXRLoadingScreen() {}
+
+	/** 
+	 * Structure describing the visual appearance of a single loading splash.
+	 * 
+	 * Splashes are shown in a tracker-relative space with the orientation reset to the direction
+	 * the player is facing when brought up.
+	 */
+	struct FSplashDesc
+	{
+		// Transform of the splash relative to the HMD orientation and location at the time of showing the loading screen.
+		FTransform	Transform = FTransform::Identity;
+		// Size of rendered quad in UE units
+		FVector2D	QuadSize = FVector2D(1.0f, 1.0f);
+		// UVs of the rendered texture
+		FBox2D		UVRect = FBox2D(FVector2D(0.0f, 0.0f), FVector2D(1.0f, 1.0f));
+		
+		// Simple animation. Rotation that will be applied to the splash every other frame.
+		FQuat		DeltaRotation = FQuat::Identity;
+
+		// If set, the splash texture will be rendered opaque regardless of the alpha channel. Not supported by all platforms.
+		bool		bIgnoreAlpha = false;
+
+		// Set to true to tell that the texture is dynamically generated and may update each frame.
+		bool		bIsDynamic	 = false;
+		
+		// Set this to true if the assigned texture is an external texture, for instance if using the media framework.
+		bool		bIsExternal	 = false;
+
+		// The texture shown. Can be set to a 2D or a Cube texture. (Cube textures may not be supported by all platforms.) 
+		FTextureRHIRef	Texture = nullptr;
+
+		// If set, overrides the texture shown for the left eye. If null, both eyes will show the same texture.
+		// Useful for stereo texture cubes.
+		FTextureRHIRef	LeftTexture = nullptr;
+	};
+
+	/**
+	 * Removes all splashes. Use this to replace the existing splashes before calling AddSplash.
+	 */
+	virtual void ClearSplashes() = 0;
+
+	/** 
+	 * Registers a splash to be shown while the loading screen is active.
+	 *
+	 * @param Splash	The settings for the new splash.
+	 */
+	virtual void AddSplash(const FSplashDesc& Splash) = 0;
+
+	/** 
+	 * Activates the loading screen.
+	 * If called while the loading screen is active, this will reinitialize the positions of all splashes according to the current HMD pose.
+	 */
+	virtual void ShowLoadingScreen() = 0;
+
+	/** 
+	 * Hides the loading screen.
+	 */
+	virtual void HideLoadingScreen() = 0;
+
+	/** 
+	 * Returns whether the loading screen is currently active or not.
+	 */
+	virtual bool IsShown() const = 0;
+
+	/**
+	 * Internal utility method for implementing backwards compatibility with IStereoLayers::Show/HideSplashScreen.
+	 * Should be called from implementations overriding IStereoLayer::UpdateSplashScreen()
+	 */
+	static void ShowLoadingScreen_Compat(bool bShow, FTexture2DRHIRef Texture, const FVector& Offset, const FVector2D& Scale);
+};
diff --git a/Engine/Source/Runtime/HeadMountedDisplay/Public/IXRTrackingSystem.h b/Engine/Source/Runtime/HeadMountedDisplay/Public/IXRTrackingSystem.h
--- a/Engine/Source/Runtime/HeadMountedDisplay/Public/IXRTrackingSystem.h
+++ b/Engine/Source/Runtime/HeadMountedDisplay/Public/IXRTrackingSystem.h
@@ -306,6 +306,14 @@ public:
 	 */
 	virtual IXRInput* GetXRInput() { return nullptr; }
 
+
+	/**
+	 * Access the loading screen interface associated with this tracking system, if any.
+	 *
+	 * @return an IXRLoadingScreen pointer or a nullptr if this tracking system does not support loading screens.
+	 */
+	virtual class IXRLoadingScreen* GetLoadingScreen() { return nullptr; }
+
 	/*** XR System related methods moved from IHeadMountedDisplay ***/
 
 	/**
diff --git a/Engine/Source/Runtime/HeadMountedDisplay/Public/StereoLayerManager.h b/Engine/Source/Runtime/HeadMountedDisplay/Public/StereoLayerManager.h
--- a/Engine/Source/Runtime/HeadMountedDisplay/Public/StereoLayerManager.h
+++ b/Engine/Source/Runtime/HeadMountedDisplay/Public/StereoLayerManager.h
@@ -5,6 +5,7 @@
 #include "CoreMinimal.h"
 #include "Misc/ScopeLock.h"
 #include "IStereoLayers.h"
+#include "IXRLoadingScreen.h"
 
 /**
 	Partial implementation of the Layer management code for the IStereoLayers interface.
@@ -31,12 +32,14 @@
 		bool GetStereoLayersDirty() -- Returns true if layer data have changed since the status was last cleared
 		ForEachLayer(...) -- pass in a lambda to iterate through each existing layer.
 		CopyLayers(TArray<LayerType>& OutArray, bool bMarkClean = true) -- Copies the layers into OutArray. 
+		WithLayer(uint32 LayerId, TFunction<void(LayerType*)> Func) -- Finds the layer by Id and passes it to Func or nullptr if not found.
 	The last two methods will clear the layer dirty flag unless you pass in false as the optional final argument.
 	
 	Thread safety:
 	Updates and the two protected access methods use a critical section to ensure atomic access to the layer structures.
 	Therefore, it is probably better to copy layers before performing time consuming operations using CopyLayers and reserve
-	ForEachLayer for simple processing or operations where you need to know the user-facing layer id.
+	ForEachLayer for simple processing or operations where you need to know the user-facing layer id. The WithLayer method
+	is useful if you already know the id of a layer you need to access in a thread safe manner.
 
 */
 
@@ -45,10 +48,43 @@ class TStereoLayerManager : public IStereoLayers
 {
 private:
 	mutable FCriticalSection LayerCritSect;
-	TMap<uint32, LayerType> StereoLayers;
-	uint32 NextLayerId;
 	bool bStereoLayersDirty;
 
+	struct FLayerData {
+		TMap<uint32, LayerType> Layers;
+		uint32 NextLayerId;
+		bool bShowBackground;
+
+		FLayerData(uint32 InNext, bool bInShowBackground = true) : Layers(), NextLayerId(InNext), bShowBackground(bInShowBackground) {}
+		FLayerData(const class FlayerData& In) : Layers(In.Layers), NextLayerId(In.NextLayerId), bShowBackground(In.bShowBackground) {}
+	};
+	TArray<FLayerData> LayerStack;
+
+	FLayerData& LayerState(int Level=0) { return LayerStack.Last(Level); }
+	//const FLayerData& LayerState(int Level = 0) const { return LayerStack.Last(Level); }
+	TMap<uint32, LayerType>& StereoLayers(int Level = 0) { return LayerState(Level).Layers; }
+	uint32 MakeLayerId() { return LayerState().NextLayerId++; }
+	
+	LayerType* FindLayerById(uint32 LayerId, int32& OutLevel)
+	{
+		if (LayerId == FLayerDesc::INVALID_LAYER_ID || LayerId >= LayerState().NextLayerId)
+		{
+			return nullptr;
+		}
+
+		// If the layer id does not exist in the current state, we'll search up the stack for the last version of the layer.
+		for (int32 I = 0; I < LayerStack.Num(); I++)
+		{
+			LayerType* Found = StereoLayers(I).Find(LayerId);
+			if (Found)
+			{
+				OutLevel = I;
+				return Found;
+			}
+		}
+		return nullptr;
+	}
+
 protected:
 
 	virtual void UpdateLayer(LayerType& Layer, uint32 LayerId, bool bIsValid) const
@@ -63,7 +99,7 @@ protected:
 	void ForEachLayer(TFunction<void(uint32, LayerType&)> Func, bool bMarkClean = true)
 	{
 		FScopeLock LockLayers(&LayerCritSect);
-		for (auto& Pair : StereoLayers)
+		for (auto& Pair : StereoLayers())
 		{
 			Func(Pair.Key, Pair.Value);
 		}
@@ -77,7 +113,7 @@ protected:
 	void CopyLayers(TArray<LayerType>& OutArray, bool bMarkClean = true)
 	{
 		FScopeLock LockLayers(&LayerCritSect);
-		StereoLayers.GenerateValueArray(OutArray);
+		StereoLayers().GenerateValueArray(OutArray);
 
 		if (bMarkClean)
 		{
@@ -88,15 +124,17 @@ protected:
 	void WithLayer(uint32 LayerId, TFunction<void(LayerType*)> Func)
 	{
 		FScopeLock LockLayers(&LayerCritSect);
-		Func(StereoLayers.Find(LayerId));
+		int32 FoundLevel;
+		Func(FindLayerById(LayerId, FoundLevel));
 	}
 
 public:
 
 	TStereoLayerManager()
-		: NextLayerId(1)
-		, bStereoLayersDirty(false)
-	{}
+		: bStereoLayersDirty(false)
+		, LayerStack{ 1 }
+	{
+	}
 
 	virtual ~TStereoLayerManager()
 	{}
@@ -105,9 +143,9 @@ public:
 	{
 		FScopeLock LockLayers(&LayerCritSect);
 
-		uint32 LayerId = NextLayerId++;
-		check(LayerId > 0);
-		LayerType& NewLayer = StereoLayers.Emplace(LayerId, InLayerDesc);
+		uint32 LayerId = MakeLayerId();
+		check(LayerId != FLayerDesc::INVALID_LAYER_ID);
+		LayerType& NewLayer = StereoLayers().Emplace(LayerId, InLayerDesc);
 		NewLayer.SetLayerId(LayerId);
 		UpdateLayer(NewLayer, LayerId, InLayerDesc.Texture != nullptr);
 		bStereoLayersDirty = true;
@@ -116,53 +154,156 @@ public:
 
 	virtual void DestroyLayer(uint32 LayerId) override
 	{
-		if (!LayerId)
+		FScopeLock LockLayers(&LayerCritSect);
+		if (LayerId == FLayerDesc::INVALID_LAYER_ID || LayerId >= LayerState().NextLayerId)
 		{
 			return;
 		}
 
-		FScopeLock LockLayers(&LayerCritSect);
-		UpdateLayer(StereoLayers[LayerId], LayerId, false);
-		StereoLayers.Remove(LayerId);
-		bStereoLayersDirty = true;
+		int32 FoundLevel;
+		LayerType* Found = FindLayerById(LayerId, FoundLevel);
+
+		// Destroy layer will delete the last active copy of the layer even if it's currently not active
+		if (Found)
+		{
+			if (FoundLevel == 0)
+			{
+				UpdateLayer(*Found, LayerId, false);
+				bStereoLayersDirty = true;
+			}
+
+			StereoLayers(FoundLevel).Remove(LayerId);
+		}
 	}
 
 	virtual void SetLayerDesc(uint32 LayerId, const FLayerDesc& InLayerDesc) override
 	{
-		if (!LayerId)
+		if (LayerId == FLayerDesc::INVALID_LAYER_ID)
 		{
 			return;
 		}
-
 		FScopeLock LockLayers(&LayerCritSect);
 
-		LayerType& Layer = StereoLayers[LayerId];
-		SetLayerDescMember(Layer, InLayerDesc);
-		Layer.SetLayerId(LayerId);
-		UpdateLayer(Layer, LayerId, InLayerDesc.Texture != nullptr);
-		bStereoLayersDirty = true;
+		// SetLayerDesc layer will update the last active copy of the layer.
+		int32 FoundLevel;
+		LayerType* Found = FindLayerById(LayerId, FoundLevel);
+		if (Found)
+		{
+			SetLayerDescMember(*Found, InLayerDesc);
+			Found->SetLayerId(LayerId);
+
+			// If the layer is currently active, update layer state
+			if (FoundLevel == 0)
+			{
+				UpdateLayer(*Found, LayerId, InLayerDesc.Texture != nullptr);
+				bStereoLayersDirty = true;
+			}
+		}
 	}
 
 	virtual bool GetLayerDesc(uint32 LayerId, FLayerDesc& OutLayerDesc) override
 	{
-		if (!LayerId)
+		if (LayerId == FLayerDesc::INVALID_LAYER_ID)
 		{
 			return false;
 		}
 		FScopeLock LockLayers(&LayerCritSect);
 
-		return GetLayerDescMember(StereoLayers[LayerId], OutLayerDesc);
+		int32 FoundLevel;
+		LayerType* Found = FindLayerById(LayerId, FoundLevel);
+		if (Found)
+		{
+			return GetLayerDescMember(*Found, OutLayerDesc);
+		}
+		else
+		{
+			return false;
+		}
 	}
 
 	virtual void MarkTextureForUpdate(uint32 LayerId) override 
 	{
-		if (!LayerId)
+		if (LayerId == FLayerDesc::INVALID_LAYER_ID)
 		{
 			return;
 		}
 		FScopeLock LockLayers(&LayerCritSect);
-		MarkLayerTextureForUpdate(StereoLayers[LayerId]);
+		// If the layer id does not exist in the current state, we'll search up the stack for the last version of the layer.
+		int32 FoundLevel;
+		LayerType* Found = FindLayerById(LayerId, FoundLevel);
+		if (Found)
+		{
+			MarkLayerTextureForUpdate(*Found);
+		}
+	}
+
+	virtual void PushLayerState(bool bPreserve = false) override
+	{
+		FScopeLock LockLayers(&LayerCritSect);
+		const auto& CurrentState = LayerState();
+
+		if (bPreserve)
+		{
+			// If bPreserve is true, copy the entire state.
+			LayerStack.Emplace(CurrentState);
+			// We don't need to mark stereo layers as dirty as the new state is a copy of the existing one.
+		}
+		else
+		{
+			// Else start with an empty set of layers, but preserve NextLayerId.
+			for (auto& Pair : StereoLayers())
+			{
+				// We need to mark the layers going out of scope as invalid, so implementations will remove them from the screen.
+				UpdateLayer(Pair.Value, Pair.Key, false);
+			}
+
+			// New layers should continue using unique layer ids
+			LayerStack.Emplace(CurrentState.NextLayerId, CurrentState.bShowBackground);
+			bStereoLayersDirty = true;
+		}
 	}
+
+	virtual void PopLayerState() override
+	{
+		FScopeLock LockLayers(&LayerCritSect);
+
+		// Ignore if there is only one element on the stack
+		if (LayerStack.Num() <= 1)
+		{
+			return;
+		}
+
+		// First mark all layers in the current state as invalid if they did not exist previously.
+		for (auto& Pair : StereoLayers(0))
+		{
+			if (!StereoLayers(1).Contains(Pair.Key))
+			{
+				UpdateLayer(Pair.Value, Pair.Key, false);
+			}
+		}
+
+		// Destroy the top of the stack
+		LayerStack.Pop();
+
+		// Update the layers in the new current state to mark them as valid and restore previous state
+		for (auto& Pair : StereoLayers(0))
+		{
+			FLayerDesc LayerDesc;
+			GetLayerDescMember(Pair.Value, LayerDesc);
+			UpdateLayer(Pair.Value, Pair.Key, LayerDesc.Texture != nullptr);
+		}
+
+		bStereoLayersDirty = true;
+	}
+
+	virtual bool SupportsLayerState() override { return true; }
+
+	virtual void HideBackgroundLayer() { LayerState().bShowBackground = false; }
+	virtual void ShowBackgroundLayer() { LayerState().bShowBackground = true; }
+	virtual bool IsBackgroundLayerVisible() const { return LayerStack.Last().bShowBackground; }
+
+	virtual void UpdateSplashScreen() override { IXRLoadingScreen::ShowLoadingScreen_Compat(bSplashIsShown, (bSplashShowMovie && SplashMovie.IsValid()) ? SplashMovie : SplashTexture, SplashOffset, SplashScale); }
+
 };
 
 class HEADMOUNTEDDISPLAY_API FSimpleLayerManager : public TStereoLayerManager<IStereoLayers::FLayerDesc>
diff --git a/Engine/Source/Runtime/HeadMountedDisplay/Public/XRLoadingScreenBase.h b/Engine/Source/Runtime/HeadMountedDisplay/Public/XRLoadingScreenBase.h
new file mode 100644
--- /dev/null
+++ b/Engine/Source/Runtime/HeadMountedDisplay/Public/XRLoadingScreenBase.h
@@ -0,0 +1,137 @@
+// Copyright 1998-2019 Epic Games, Inc. All Rights Reserved.
+
+#pragma once
+
+#include "IXRLoadingScreen.h"
+#include "TickableObjectRenderThread.h"
+
+
+
+/**
+* Base utility class for implementations of the IXRLoadingScreen interface
+*/
+template<typename SplashType>
+class HEADMOUNTEDDISPLAY_API TXRLoadingScreenBase : public IXRLoadingScreen, protected FTickableObjectRenderThread
+{
+public: 
+
+	TXRLoadingScreenBase(class IXRTrackingSystem* InTrackingSystem)
+		: FTickableObjectRenderThread(false, false)
+		, TrackingSystem(InTrackingSystem)
+		, bShowing(false)
+		, SystemDisplayInterval(1 / 90.0f)
+	{
+		check(TrackingSystem);
+	}
+
+	/* IXRLoadingScreen interface */
+	virtual void ClearSplashes() override
+	{
+		if (bShowing)
+		{
+			for (SplashType& Splash : Splashes)
+			{
+				DoDeleteSplash(Splash);
+			}
+		}
+		Splashes.Reset();
+	}
+
+	virtual void AddSplash(const FSplashDesc& Splash) override
+	{
+		Splashes.Emplace(Splash);
+		DoAddSplash(Splashes.Last());
+	}
+
+	virtual void ShowLoadingScreen() override
+	{
+		{
+			FQuat TMPOrientation = FQuat::Identity;
+			FVector TMPPosition;
+			TrackingSystem->GetCurrentPose(IXRTrackingSystem::HMDDeviceId, TMPOrientation, TMPPosition);
+
+			// Only use the Yaw of the HMD orientation
+			FRotator TMPRotation(TMPOrientation);
+			TMPRotation.Pitch = 0;
+			TMPRotation.Roll = 0;
+
+			HMDOrientation = FQuat(TMPRotation);
+			HMDOrientation.Normalize();
+		}
+
+		for (FSplashData& Splash : Splashes)
+		{
+			DoShowSplash(Splash);
+		}
+
+		if (!bShowing)
+		{
+			bShowing = true;
+			ExecuteOnRenderThread_DoNotWait([this]()
+			{
+				Register();
+			});
+		}
+	}
+
+	virtual void HideLoadingScreen() override
+	{
+		if (!bShowing)
+		{
+			return;
+		}
+
+		for (FSplashData& Splash : Splashes)
+		{
+			DoHideSplash(Splash);
+		}
+
+
+		ExecuteOnRenderThread([this]()
+			{
+				Unregister();
+			});
+
+		bShowing = false;
+	}
+
+	virtual bool IsShown() const override { return bShowing; }
+
+
+protected:
+
+	/* FTickableObjectRenderThread interface */
+	virtual void Tick(float DeltaTime) override
+	{
+		const double TimeInSeconds = FPlatformTime::Seconds();
+		const double DeltaTimeInSeconds = TimeInSeconds - LastTimeInSeconds;
+
+		if (DeltaTimeInSeconds > 2.f * SystemDisplayInterval)
+		{
+			for (SplashType& Splash : Splashes)
+			{
+				ApplyDeltaRotation(Splash);
+			}
+			LastTimeInSeconds = TimeInSeconds;
+		}
+
+	}
+
+	virtual TStatId GetStatId() const override { RETURN_QUICK_DECLARE_CYCLE_STAT(FDefaultXRLoadingScreen, STATGROUP_Tickables); }
+	virtual bool IsTickable() const override { return true; }
+
+	virtual void DoShowSplash(SplashType& Splash) = 0;
+	virtual void DoHideSplash(SplashType& Splash) = 0;
+	virtual void DoDeleteSplash(SplashType& Splash) = 0;
+	virtual void DoAddSplash(SplashType& Splash) = 0;
+	virtual void ApplyDeltaRotation(const SplashType& Splash) = 0;
+
+	TArray<SplashType> Splashes;
+
+	class IXRTrackingSystem* TrackingSystem;
+	FQuat HMDOrientation;
+	bool bShowing;
+
+	float SystemDisplayInterval;
+	double LastTimeInSeconds;
+};
\ No newline at end of file
diff --git a/Engine/Source/Runtime/HeadMountedDisplay/Public/XRLoadingScreenFunctionLibrary.h b/Engine/Source/Runtime/HeadMountedDisplay/Public/XRLoadingScreenFunctionLibrary.h
new file mode 100644
--- /dev/null
+++ b/Engine/Source/Runtime/HeadMountedDisplay/Public/XRLoadingScreenFunctionLibrary.h
@@ -0,0 +1,52 @@
+// Copyright 1998-2019 Epic Games, Inc. All Rights Reserved.
+
+#pragma once
+
+#include "CoreMinimal.h"
+#include "UObject/ObjectMacros.h"
+#include "Kismet/BlueprintFunctionLibrary.h"
+#include "Kismet/BlueprintAsyncActionBase.h"
+#include "IXRLoadingScreen.h"
+#include "XRLoadingScreenFunctionLibrary.generated.h"
+
+
+
+/**
+ * XR Loading Screen Function Library 
+ */
+UCLASS()
+class HEADMOUNTEDDISPLAY_API UXRLoadingScreenFunctionLibrary : public UBlueprintFunctionLibrary
+{
+	GENERATED_UCLASS_BODY()
+
+	UFUNCTION(BlueprintCallable, Category = "Input|HeadMountedDisplay|LoadingScreen")
+	static void SetLoadingScreen(class UTexture* Texture, FVector2D Scale, FVector Offset, bool bShowLoadingMovie, bool bShowOnSet);
+	
+	UFUNCTION(BlueprintCallable, Category = "Input|HeadMountedDisplay|LoadingScreen")
+	static void ClearLoadingScreenSplashes();
+
+	/**
+     * Adds a splash element to the loading screen.
+     *
+	 * @param Texture			(in) A texture asset to be used for the splash.
+	 * @param Translation       (in) Initial translation of the center of the splash.
+	 * @param Rotation			(in) Initial rotation of the splash screen, with the origin at the center of the splash.
+	 * @param Size      		(in) Size, of the quad with the splash screen.
+	 * @param DeltaRotation		(in) Incremental rotation, that is added each 2nd frame to the quad transform. The quad is rotated around the center of the quad.
+	 * @param bClearBeforeAdd	(in) If true, clears splashes before adding a new one.
+	 */
+	UFUNCTION(BlueprintCallable, Category = "Input|HeadMountedDisplay|LoadingScreen")
+	static void AddLoadingScreenSplash(class UTexture* Texture, FVector Translation, FRotator Rotation, FVector2D Size = FVector2D(1.0f, 1.0f), FRotator DeltaRotation = FRotator::ZeroRotator, bool bClearBeforeAdd = false);
+	
+	/**
+	 * Show the loading screen and override the VR display
+	 */	
+	UFUNCTION(BlueprintCallable, Category = "Input|HeadMountedDisplay|LoadingScreen")
+	static void ShowLoadingScreen();
+
+	/**
+	 * Hide the splash screen and return to normal display.
+	 */
+	UFUNCTION(BlueprintCallable, Category = "Input|HeadMountedDisplay|LoadingScreen")
+	static void HideLoadingScreen();
+};
\ No newline at end of file
diff --git a/Engine/Source/Runtime/HeadMountedDisplay/Public/XRTrackingSystemBase.h b/Engine/Source/Runtime/HeadMountedDisplay/Public/XRTrackingSystemBase.h
--- a/Engine/Source/Runtime/HeadMountedDisplay/Public/XRTrackingSystemBase.h
+++ b/Engine/Source/Runtime/HeadMountedDisplay/Public/XRTrackingSystemBase.h
@@ -138,12 +138,18 @@ public:
 	* @param ExternalTrackingTransform		The transform in world-coordinates, of the reference marker of the external tracking system
 	*/
 	virtual void UpdateExternalTrackingPosition(const FTransform& ExternalTrackingTransform) override;
+	virtual class IXRLoadingScreen* GetLoadingScreen() override final;
 
 
 	TSharedPtr<FARSupportInterface , ESPMode::ThreadSafe> GetARCompositionComponent();
 	const TSharedPtr<const FARSupportInterface , ESPMode::ThreadSafe> GetARCompositionComponent() const;
 
 protected:
+
+	/** Override this method if you need to customize the IXRLoadingScreen implementation.
+	 *  It will be called from GetLoadingScreen() if LoadingScreen is not valid */
+	virtual class IXRLoadingScreen* CreateLoadingScreen();
+
 	/** 
 	 * Meant to be called by sub-classes whenever the tracking origin is altered.
 	 */
@@ -173,6 +179,7 @@ protected:
 	 * will attempt to normalize the internal tracking system to match this calibration when called.
 	 */
 	FTransform CalibratedOffset;
+	mutable class IXRLoadingScreen* LoadingScreen;
 
 private:
 	TSharedPtr<FARSupportInterface , ESPMode::ThreadSafe> ARCompositionComponent;
