Fix iOS and/or non-unity compile issues. (plus a single copy-paste error)

diff --git a/Engine/Source/Runtime/HeadMountedDisplay/Private/DefaultXRLoadingScreen.cpp b/Engine/Source/Runtime/HeadMountedDisplay/Private/DefaultXRLoadingScreen.cpp
--- a/Engine/Source/Runtime/HeadMountedDisplay/Private/DefaultXRLoadingScreen.cpp
+++ b/Engine/Source/Runtime/HeadMountedDisplay/Private/DefaultXRLoadingScreen.cpp
@@ -2,7 +2,6 @@
 
 #include "DefaultXRLoadingScreen.h"
 #include "IXRTrackingSystem.h"
-#include "XRThreadUtils.h"
 
 using FLayerDesc = IStereoLayers::FLayerDesc;
 using FSplashDesc = IXRLoadingScreen::FSplashDesc;
diff --git a/Engine/Source/Runtime/HeadMountedDisplay/Public/StereoLayerManager.h b/Engine/Source/Runtime/HeadMountedDisplay/Public/StereoLayerManager.h
--- a/Engine/Source/Runtime/HeadMountedDisplay/Public/StereoLayerManager.h
+++ b/Engine/Source/Runtime/HeadMountedDisplay/Public/StereoLayerManager.h
@@ -56,7 +56,7 @@ private:
 		bool bShowBackground;
 
 		FLayerData(uint32 InNext, bool bInShowBackground = true) : Layers(), NextLayerId(InNext), bShowBackground(bInShowBackground) {}
-		FLayerData(const class FlayerData& In) : Layers(In.Layers), NextLayerId(In.NextLayerId), bShowBackground(In.bShowBackground) {}
+		FLayerData(const FlayerData& In) : Layers(In.Layers), NextLayerId(In.NextLayerId), bShowBackground(In.bShowBackground) {}
 	};
 	TArray<FLayerData> LayerStack;
 
diff --git a/Engine/Source/Runtime/HeadMountedDisplay/Public/XRLoadingScreenBase.h b/Engine/Source/Runtime/HeadMountedDisplay/Public/XRLoadingScreenBase.h
--- a/Engine/Source/Runtime/HeadMountedDisplay/Public/XRLoadingScreenBase.h
+++ b/Engine/Source/Runtime/HeadMountedDisplay/Public/XRLoadingScreenBase.h
@@ -4,7 +4,7 @@
 
 #include "IXRLoadingScreen.h"
 #include "TickableObjectRenderThread.h"
-
+#include "XRThreadUtils.h"
 
 
 /**
@@ -59,7 +59,7 @@ public:
 			HMDOrientation.Normalize();
 		}
 
-		for (FSplashData& Splash : Splashes)
+		for (SplashType& Splash : Splashes)
 		{
 			DoShowSplash(Splash);
 		}
@@ -81,7 +81,7 @@ public:
 			return;
 		}
 
-		for (FSplashData& Splash : Splashes)
+		for (SplashType& Splash : Splashes)
 		{
 			DoHideSplash(Splash);
 		}
