Fixing capitalization

diff --git a/Engine/Source/Runtime/HeadMountedDisplay/Public/StereoLayerManager.h b/Engine/Source/Runtime/HeadMountedDisplay/Public/StereoLayerManager.h
--- a/Engine/Source/Runtime/HeadMountedDisplay/Public/StereoLayerManager.h
+++ b/Engine/Source/Runtime/HeadMountedDisplay/Public/StereoLayerManager.h
@@ -56,7 +56,7 @@ private:
 		bool bShowBackground;
 
 		FLayerData(uint32 InNext, bool bInShowBackground = true) : Layers(), NextLayerId(InNext), bShowBackground(bInShowBackground) {}
-		FLayerData(const FlayerData& In) : Layers(In.Layers), NextLayerId(In.NextLayerId), bShowBackground(In.bShowBackground) {}
+		FLayerData(const FLayerData& In) : Layers(In.Layers), NextLayerId(In.NextLayerId), bShowBackground(In.bShowBackground) {}
 	};
 	TArray<FLayerData> LayerStack;
 
