Fix unity build error

diff --git a/Engine/Source/Runtime/HeadMountedDisplay/Public/XRLoadingScreenBase.h b/Engine/Source/Runtime/HeadMountedDisplay/Public/XRLoadingScreenBase.h
--- a/Engine/Source/Runtime/HeadMountedDisplay/Public/XRLoadingScreenBase.h
+++ b/Engine/Source/Runtime/HeadMountedDisplay/Public/XRLoadingScreenBase.h
@@ -5,7 +5,7 @@
 #include "IXRLoadingScreen.h"
 #include "TickableObjectRenderThread.h"
 #include "XRThreadUtils.h"
-
+#include "IXRTrackingSystem.h"
 
 /**
 * Base utility class for implementations of the IXRLoadingScreen interface
