Merging using PartnerOculus_to_Release_4.24

[Partner-Oculus] Fix Late update threading sync issues

1. Change " for (auto PrimitivePair : ... ) " to "for (auto& PrimitivePair :",  since we want to use " PrimitivePair.Value = -1" to mark those primitive for avoiding double updating.
2.  If  FLateUpdateManager::Setup was called earlier (in main thread) than FLateUpdateManager::Apply_RenderThread ( in render thread),  then  it can cause controller renderPose 1 frame faster than it should be,  which will result hiccups,  this is rare but possible  escpecially when main thread is super fast.  The fix use a similar data structure with 4.22, but use ENQUEUE_RENDER_COMMAND to fix the potential LateUpdateGameWriteIndex collision issue.  this also helps removed the critical section.

[CL 8618530 by JianR.Zhang in Partner-Oculus-4.23-Private]
[CL 9374881 by [REDACTED_PERSON] in Partner-Oculus-Private]

diff --git a/Engine/Source/Runtime/HeadMountedDisplay/Private/LateUpdateManager.cpp b/Engine/Source/Runtime/HeadMountedDisplay/Private/LateUpdateManager.cpp
--- a/Engine/Source/Runtime/HeadMountedDisplay/Private/LateUpdateManager.cpp
+++ b/Engine/Source/Runtime/HeadMountedDisplay/Private/LateUpdateManager.cpp
@@ -7,6 +7,8 @@
 #include "HeadMountedDisplayTypes.h"
 
 FLateUpdateManager::FLateUpdateManager()
+	: LateUpdateGameWriteIndex(0)
+	, LateUpdateRenderReadIndex(0)
 {
 }
 
@@ -14,45 +16,41 @@ void FLateUpdateManager::Setup(const FTransform& ParentToWorld, USceneComponent*
 {
 	check(IsInGameThread());
 
-	FScopeLock Lock(&StateCriticalSection);
-
-	GameThreadState.ParentToWorld = ParentToWorld;
+	UpdateStates[LateUpdateGameWriteIndex].Primitives.Reset();
+	UpdateStates[LateUpdateGameWriteIndex].ParentToWorld = ParentToWorld;
 	GatherLateUpdatePrimitives(Component);
-	GameThreadState.bSkip = bSkipLateUpdate;
-	++GameThreadState.TrackingNumber;
-}
+	UpdateStates[LateUpdateGameWriteIndex].bSkip = bSkipLateUpdate;
+	++UpdateStates[LateUpdateGameWriteIndex].TrackingNumber;
 
-void FLateUpdateManager::Apply_RenderThread(FSceneInterface* Scene, const FTransform& OldRelativeTransform, const FTransform& NewRelativeTransform)
-{
-	check(IsInRenderingThread());
+	int32 NextFrameRenderReadIndex = LateUpdateGameWriteIndex;
+	LateUpdateGameWriteIndex = 1 - LateUpdateGameWriteIndex;
 
+	ENQUEUE_RENDER_COMMAND(UpdateLateUpdateRenderReadIndexCommand)(
+		[NextFrameRenderReadIndex, this](FRHICommandListImmediate& RHICmdList)
 	{
-		// Only grab the lock to protect access to the GameThreadState collection - the RenderThreadState is
-		// only ever accessed from the rendering thread.
-		FScopeLock Lock(&StateCriticalSection);
-		RenderThreadState.bSkip = GameThreadState.bSkip;
-		RenderThreadState.ParentToWorld = GameThreadState.ParentToWorld;
-		RenderThreadState.Primitives = MoveTemp(GameThreadState.Primitives);
+		LateUpdateRenderReadIndex = NextFrameRenderReadIndex;
+	});
 
-		GameThreadState.Primitives.Reset();
+}
 
-		RenderThreadState.TrackingNumber = GameThreadState.TrackingNumber;
-	}
+void FLateUpdateManager::Apply_RenderThread(FSceneInterface* Scene, const FTransform& OldRelativeTransform, const FTransform& NewRelativeTransform)
+{
+	check(IsInRenderingThread());
 
-	if (!RenderThreadState.Primitives.Num() || RenderThreadState.bSkip)
+	if (!UpdateStates[LateUpdateRenderReadIndex].Primitives.Num() || UpdateStates[LateUpdateRenderReadIndex].bSkip)
 	{
 		return;
 	}
 
-	const FTransform OldCameraTransform = OldRelativeTransform * RenderThreadState.ParentToWorld;
-	const FTransform NewCameraTransform = NewRelativeTransform * RenderThreadState.ParentToWorld;
+	const FTransform OldCameraTransform = OldRelativeTransform * UpdateStates[LateUpdateRenderReadIndex].ParentToWorld;
+	const FTransform NewCameraTransform = NewRelativeTransform * UpdateStates[LateUpdateRenderReadIndex].ParentToWorld;
 	const FMatrix LateUpdateTransform = (OldCameraTransform.Inverse() * NewCameraTransform).ToMatrixWithScale();
 
 	bool bIndicesHaveChanged = false;
 
 	// Apply delta to the cached scene proxies
 	// Also check whether any primitive indices have changed, in case the scene has been modified in the meantime.
-	for (auto PrimitivePair : RenderThreadState.Primitives)
+	for (auto& PrimitivePair : UpdateStates[LateUpdateRenderReadIndex].Primitives)
 	{
 		FPrimitiveSceneInfo* RetrievedSceneInfo = Scene->GetPrimitiveSceneInfo(PrimitivePair.Value);
 		FPrimitiveSceneInfo* CachedSceneInfo = PrimitivePair.Key;
@@ -79,7 +77,7 @@ void FLateUpdateManager::Apply_RenderThread(FSceneInterface* Scene, const FTrans
 		RetrievedSceneInfo = Scene->GetPrimitiveSceneInfo(Index++);
 		while(RetrievedSceneInfo)
 		{
-			if (RetrievedSceneInfo->Proxy && RenderThreadState.Primitives.Contains(RetrievedSceneInfo) && RenderThreadState.Primitives[RetrievedSceneInfo] >= 0)
+			if (RetrievedSceneInfo->Proxy && UpdateStates[LateUpdateRenderReadIndex].Primitives.Contains(RetrievedSceneInfo) && UpdateStates[LateUpdateRenderReadIndex].Primitives[RetrievedSceneInfo] >= 0)
 			{
 				RetrievedSceneInfo->Proxy->ApplyLateUpdateTransform(LateUpdateTransform);
 			}
@@ -97,7 +95,7 @@ void FLateUpdateManager::CacheSceneInfo(USceneComponent* Component)
 		FPrimitiveSceneInfo* PrimitiveSceneInfo = PrimitiveComponent->SceneProxy->GetPrimitiveSceneInfo();
 		if (PrimitiveSceneInfo && PrimitiveSceneInfo->IsIndexValid())
 		{
-			GameThreadState.Primitives.Emplace(PrimitiveSceneInfo, PrimitiveSceneInfo->GetIndex());
+			UpdateStates[LateUpdateGameWriteIndex].Primitives.Emplace(PrimitiveSceneInfo, PrimitiveSceneInfo->GetIndex());
 		}
 	}
 }
diff --git a/Engine/Source/Runtime/HeadMountedDisplay/Public/LateUpdateManager.h b/Engine/Source/Runtime/HeadMountedDisplay/Public/LateUpdateManager.h
--- a/Engine/Source/Runtime/HeadMountedDisplay/Public/LateUpdateManager.h
+++ b/Engine/Source/Runtime/HeadMountedDisplay/Public/LateUpdateManager.h
@@ -20,7 +20,7 @@ public:
 	void Setup(const FTransform& ParentToWorld, USceneComponent* Component, bool bSkipLateUpdate);
 
 	/** Returns true if the LateUpdateSetup data is stale. */
-	bool GetSkipLateUpdate_RenderThread() const { return RenderThreadState.bSkip; }
+	bool GetSkipLateUpdate_RenderThread() const { return UpdateStates[LateUpdateRenderReadIndex].bSkip; }
 
 	/** Apply the late update delta to the cached components */
 	void Apply_RenderThread(FSceneInterface* Scene, const FTransform& OldRelativeTransform, const FTransform& NewRelativeTransform);
@@ -50,8 +50,8 @@ private:
 		int64 TrackingNumber;
 	};
 
-	FLateUpdateState GameThreadState;
-	FLateUpdateState RenderThreadState;
-	FCriticalSection StateCriticalSection;
+	FLateUpdateState UpdateStates[2];
+	int32 LateUpdateGameWriteIndex;
+	int32 LateUpdateRenderReadIndex;
 };
 
