Add missing transition in DefaultStereoLayers

diff --git a/Engine/Source/Runtime/HeadMountedDisplay/Private/DefaultStereoLayers.cpp b/Engine/Source/Runtime/HeadMountedDisplay/Private/DefaultStereoLayers.cpp
--- a/Engine/Source/Runtime/HeadMountedDisplay/Private/DefaultStereoLayers.cpp
+++ b/Engine/Source/Runtime/HeadMountedDisplay/Private/DefaultStereoLayers.cpp
@@ -138,6 +138,8 @@ void FDefaultStereoLayers::StereoLayerRender(FRHICommandListImmediate& RHICmdLis
 			}
 		}
 
+		RHICmdList.TransitionResource(EResourceTransitionAccess::EReadable, Layer.Texture);
+
 		// Set shader uniforms
 		VertexShader->SetParameters(
 			RHICmdList,
@@ -290,4 +292,4 @@ void FDefaultStereoLayers::SetupViewFamily(FSceneViewFamily& InViewFamily)
 	FVector HmdPosition = FVector::ZeroVector;
 	HMDDevice->GetCurrentPose(IXRTrackingSystem::HMDDeviceId, HmdOrientation, HmdPosition);
 	HmdTransform = FTransform(HmdOrientation, HmdPosition);
-}
\ No newline at end of file
+}
