XR: Fix static analysis complaints about HiddenAreaMesh population

Bonus: Static analysis complained about FSteamVRHMD, fixed same pattern in FHMDViewMesh

diff --git a/Engine/Plugins/Runtime/Steam/SteamVR/Source/SteamVR/Private/SteamVRHMD.cpp b/Engine/Plugins/Runtime/Steam/SteamVR/Source/SteamVR/Private/SteamVRHMD.cpp
--- a/Engine/Plugins/Runtime/Steam/SteamVR/Source/SteamVR/Private/SteamVRHMD.cpp
+++ b/Engine/Plugins/Runtime/Steam/SteamVR/Source/SteamVR/Private/SteamVRHMD.cpp
@@ -1916,25 +1916,19 @@ static void SetupHiddenAreaMeshes(vr::IVRSystem* const VRSystem, FHMDViewMesh Re
 		FVector2D* const LeftEyePositions = new FVector2D[VertexCount];
 		FVector2D* const RightEyePositions = new FVector2D[VertexCount];
 
-		uint32 DataIndex = 0;
-		for (uint32 TriangleIter = 0; TriangleIter < LeftEyeMesh.unTriangleCount; ++TriangleIter)
+		for (uint32 VertexIndex = 0; VertexIndex < VertexCount; ++VertexIndex)
 		{
-			for (uint32 VertexIter = 0; VertexIter < 3; ++VertexIter)
-			{
-				const vr::HmdVector2_t& LeftSrc = LeftEyeMesh.pVertexData[DataIndex];
-				const vr::HmdVector2_t& RightSrc = RightEyeMesh.pVertexData[DataIndex];
-
-				FVector2D& LeftDst = LeftEyePositions[DataIndex];
-				FVector2D& RightDst = RightEyePositions[DataIndex];
+			const vr::HmdVector2_t& LeftSrc = LeftEyeMesh.pVertexData[VertexIndex];
+			const vr::HmdVector2_t& RightSrc = RightEyeMesh.pVertexData[VertexIndex];
 
-				LeftDst.X = LeftSrc.v[0];
-				LeftDst.Y = LeftSrc.v[1];
+			FVector2D& LeftDst = LeftEyePositions[VertexIndex];
+			FVector2D& RightDst = RightEyePositions[VertexIndex];
 
-				RightDst.X = RightSrc.v[0];
-				RightDst.Y = RightSrc.v[1];
+			LeftDst.X = LeftSrc.v[0];
+			LeftDst.Y = LeftSrc.v[1];
 
-				++DataIndex;
-			}
+			RightDst.X = RightSrc.v[0];
+			RightDst.Y = RightSrc.v[1];
 		}
 		
 		const FHMDViewMesh::EHMDMeshType MeshTransformType = (MeshType == vr::EHiddenAreaMeshType::k_eHiddenAreaMesh_Standard) ? FHMDViewMesh::MT_HiddenArea : FHMDViewMesh::MT_VisibleArea;
diff --git a/Engine/Source/Runtime/HeadMountedDisplay/Private/HeadMountedDisplayTypes.cpp b/Engine/Source/Runtime/HeadMountedDisplay/Private/HeadMountedDisplayTypes.cpp
--- a/Engine/Source/Runtime/HeadMountedDisplay/Private/HeadMountedDisplayTypes.cpp
+++ b/Engine/Source/Runtime/HeadMountedDisplay/Private/HeadMountedDisplayTypes.cpp
@@ -34,42 +34,36 @@ void FHMDViewMesh::BuildMesh(const FVector2D Positions[], uint32 VertexCount, EH
 	void* VoidPtr2 = RHILockBuffer(IndexBufferRHI, 0, sizeof(uint16) * NumIndices, RLM_WriteOnly);
 	uint16* pIndices = reinterpret_cast<uint16*>(VoidPtr2);
 
-	uint32 DataIndex = 0;
-	for (uint32 TriangleIter = 0; TriangleIter < NumTriangles; ++TriangleIter)
+	for (uint32 VertexIndex = 0; VertexIndex < NumVertices; ++VertexIndex)
 	{
-		for (uint32 VertexIter = 0; VertexIter < 3; ++VertexIter)
-		{
-			const FVector2D& Position = Positions[DataIndex];
-			FFilterVertex& Vertex = pVertices[DataIndex];
-
-			if (MeshType == MT_HiddenArea)
-			{
-				// Remap from to NDC space [0 1] -> [-1 1]
-				Vertex.Position.X = (Position.X * 2.0f) - 1.0f;
-				Vertex.Position.Y = (Position.Y * 2.0f) - 1.0f;
-				Vertex.Position.Z = 1.0f;
-				Vertex.Position.W = 1.0f;
+		const FVector2D& Position = Positions[VertexIndex];
+		FFilterVertex& Vertex = pVertices[VertexIndex];
 
-				// Not used for hidden area
-				Vertex.UV.X = 0.0f;
-				Vertex.UV.Y = 0.0f;
-			}
-			else
-			{
-				// Remap the viewport origin from the bottom left to the top left
-				Vertex.Position.X = Position.X;
-				Vertex.Position.Y = 1.0f - Position.Y;
-				Vertex.Position.Z = 0.0f;
-				Vertex.Position.W = 1.0f;
-
-				Vertex.UV.X = Position.X;
-				Vertex.UV.Y = 1.0f - Position.Y;
-			}
+		if (MeshType == MT_HiddenArea)
+		{
+			// Remap from to NDC space [0 1] -> [-1 1]
+			Vertex.Position.X = (Position.X * 2.0f) - 1.0f;
+			Vertex.Position.Y = (Position.Y * 2.0f) - 1.0f;
+			Vertex.Position.Z = 1.0f;
+			Vertex.Position.W = 1.0f;
 
-			pIndices[DataIndex] = DataIndex;
+			// Not used for hidden area
+			Vertex.UV.X = 0.0f;
+			Vertex.UV.Y = 0.0f;
+		}
+		else
+		{
+			// Remap the viewport origin from the bottom left to the top left
+			Vertex.Position.X = Position.X;
+			Vertex.Position.Y = 1.0f - Position.Y;
+			Vertex.Position.Z = 0.0f;
+			Vertex.Position.W = 1.0f;
 
-			++DataIndex;
+			Vertex.UV.X = Position.X;
+			Vertex.UV.Y = 1.0f - Position.Y;
 		}
+
+		pIndices[VertexIndex] = VertexIndex;
 	}
 
 	RHIUnlockBuffer(VertexBufferRHI);
