Add GetPlayAreaBounds() blueprint function

diff --git a/Engine/Plugins/Runtime/Oculus/OculusVR/Source/OculusHMD/Private/OculusHMD.cpp b/Engine/Plugins/Runtime/Oculus/OculusVR/Source/OculusHMD/Private/OculusHMD.cpp
--- a/Engine/Plugins/Runtime/Oculus/OculusVR/Source/OculusHMD/Private/OculusHMD.cpp
+++ b/Engine/Plugins/Runtime/Oculus/OculusVR/Source/OculusHMD/Private/OculusHMD.cpp
@@ -1023,6 +1023,20 @@ namespace OculusHMD
 		return true;
 	}
 
+	FVector2D FOculusHMD::GetPlayAreaBounds(EHMDTrackingOrigin::Type Origin) const
+	{
+		ovrpVector3f Dimensions;
+
+		if (Origin == EHMDTrackingOrigin::Stage &&
+			OVRP_SUCCESS(FOculusHMDModule::GetPluginWrapper().GetBoundaryDimensions2(ovrpBoundary_PlayArea, &Dimensions)))
+		{
+			Dimensions.z *= -1.0;
+			FVector Bounds = ConvertVector_M2U(Dimensions);
+			return FVector2D(Bounds.X, Bounds.Z);
+		}
+		return FVector2D::ZeroVector;
+	}
+
 	bool FOculusHMD::IsHMDConnected()
 	{
 		CheckInGameThread();
diff --git a/Engine/Plugins/Runtime/Oculus/OculusVR/Source/OculusHMD/Private/OculusHMD.h b/Engine/Plugins/Runtime/Oculus/OculusVR/Source/OculusHMD/Private/OculusHMD.h
--- a/Engine/Plugins/Runtime/Oculus/OculusVR/Source/OculusHMD/Private/OculusHMD.h
+++ b/Engine/Plugins/Runtime/Oculus/OculusVR/Source/OculusHMD/Private/OculusHMD.h
@@ -109,6 +109,7 @@ public:
 	virtual bool OnStartGameFrame(FWorldContext& WorldContext) override;
 	virtual bool OnEndGameFrame(FWorldContext& WorldContext) override;
 	virtual class IXRLoadingScreen* CreateLoadingScreen() override { return GetSplash(); }
+	virtual FVector2D GetPlayAreaBounds(EHMDTrackingOrigin::Type Origin) const override;
 
 	// IHeadMountedDisplay
 	virtual bool IsHMDConnected() override;
diff --git a/Engine/Plugins/Runtime/OpenXR/Source/OpenXRHMD/Private/OpenXRHMD.cpp b/Engine/Plugins/Runtime/OpenXR/Source/OpenXRHMD/Private/OpenXRHMD.cpp
--- a/Engine/Plugins/Runtime/OpenXR/Source/OpenXRHMD/Private/OpenXRHMD.cpp
+++ b/Engine/Plugins/Runtime/OpenXR/Source/OpenXRHMD/Private/OpenXRHMD.cpp
@@ -882,6 +882,28 @@ float FOpenXRHMD::GetWorldToMetersScale() const
 	return WorldToMetersScale;
 }
 
+FVector2D FOpenXRHMD::GetPlayAreaBounds(EHMDTrackingOrigin::Type Origin) const
+{
+	XrReferenceSpaceType Space = XR_REFERENCE_SPACE_TYPE_STAGE;
+	switch (Origin)
+	{
+	case EHMDTrackingOrigin::Eye:
+		Space = XR_REFERENCE_SPACE_TYPE_VIEW;
+		break;
+	case EHMDTrackingOrigin::Floor:
+		Space = XR_REFERENCE_SPACE_TYPE_LOCAL;
+		break;
+	case EHMDTrackingOrigin::Stage:
+		Space = XR_REFERENCE_SPACE_TYPE_STAGE;
+		break;
+	default:
+		break;
+	}
+	XrExtent2Df Bounds;
+	XR_ENSURE(xrGetReferenceSpaceBoundsRect(Session, Space, &Bounds));
+	return ToFVector2D(Bounds, WorldToMetersScale);
+}
+
 bool FOpenXRHMD::IsHMDEnabled() const
 {
 	return true;
diff --git a/Engine/Plugins/Runtime/OpenXR/Source/OpenXRHMD/Private/OpenXRHMD.h b/Engine/Plugins/Runtime/OpenXR/Source/OpenXRHMD/Private/OpenXRHMD.h
--- a/Engine/Plugins/Runtime/OpenXR/Source/OpenXRHMD/Private/OpenXRHMD.h
+++ b/Engine/Plugins/Runtime/OpenXR/Source/OpenXRHMD/Private/OpenXRHMD.h
@@ -162,6 +162,8 @@ public:
 
 	virtual float GetWorldToMetersScale() const override;
 
+	virtual FVector2D GetPlayAreaBounds(EHMDTrackingOrigin::Type Origin) const override;
+
 protected:
 
 	bool StartSession();
diff --git a/Engine/Plugins/Runtime/OpenXR/Source/OpenXRHMD/Public/OpenXRCore.h b/Engine/Plugins/Runtime/OpenXR/Source/OpenXRHMD/Public/OpenXRCore.h
--- a/Engine/Plugins/Runtime/OpenXR/Source/OpenXRHMD/Public/OpenXRCore.h
+++ b/Engine/Plugins/Runtime/OpenXR/Source/OpenXRHMD/Public/OpenXRCore.h
@@ -81,6 +81,11 @@ FORCEINLINE XrRect2Di ToXrRect(FIntRect Rect)
 	return XrRect2Di{ { Rect.Min.X, Rect.Min.Y }, { Rect.Width(), Rect.Height() } };
 }
 
+FORCEINLINE FVector2D ToFVector2D(XrExtent2Df Extent, float Scale = 1.0f)
+{
+	return FVector2D(Extent.width * Scale, Extent.height * Scale);
+}
+
 FORCEINLINE XrExtent2Df ToXrExtent2D(FVector2D Vector, float Scale = 1.0f)
 {
 	if (Vector.IsZero())
diff --git a/Engine/Plugins/Runtime/Steam/SteamVR/Source/SteamVR/Private/SteamVRHMD.cpp b/Engine/Plugins/Runtime/Steam/SteamVR/Source/SteamVR/Private/SteamVRHMD.cpp
--- a/Engine/Plugins/Runtime/Steam/SteamVR/Source/SteamVR/Private/SteamVRHMD.cpp
+++ b/Engine/Plugins/Runtime/Steam/SteamVR/Source/SteamVR/Private/SteamVRHMD.cpp
@@ -897,6 +897,16 @@ bool FSteamVRHMD::GetFloorToEyeTrackingTransform(FTransform& OutStandingToSeated
 	return bSuccess;
 }
 
+FVector2D FSteamVRHMD::GetPlayAreaBounds(EHMDTrackingOrigin::Type Origin) const
+{
+	FVector2D Bounds;
+	if (Origin == EHMDTrackingOrigin::Stage && VRChaperone->GetPlayAreaSize(&Bounds.X, &Bounds.Y))
+	{
+		return Bounds;
+	}
+	return FVector2D::ZeroVector;
+}
+
 void FSteamVRHMD::RecordAnalytics()
 {
 	if (FEngineAnalytics::IsAvailable())
diff --git a/Engine/Plugins/Runtime/Steam/SteamVR/Source/SteamVR/Private/SteamVRHMD.h b/Engine/Plugins/Runtime/Steam/SteamVR/Source/SteamVR/Private/SteamVRHMD.h
--- a/Engine/Plugins/Runtime/Steam/SteamVR/Source/SteamVR/Private/SteamVRHMD.h
+++ b/Engine/Plugins/Runtime/Steam/SteamVR/Source/SteamVR/Private/SteamVRHMD.h
@@ -124,6 +124,7 @@ public:
 	virtual void SetTrackingOrigin(EHMDTrackingOrigin::Type NewOrigin) override;
 	virtual EHMDTrackingOrigin::Type GetTrackingOrigin() const override;
 	virtual bool GetFloorToEyeTrackingTransform(FTransform& OutFloorToEye) const override;
+	virtual FVector2D GetPlayAreaBounds(EHMDTrackingOrigin::Type Origin) const override;
 
 public:
 	/** IHeadMountedDisplay interface */
diff --git a/Engine/Source/Runtime/HeadMountedDisplay/Private/HeadMountedDisplayFunctionLibrary.cpp b/Engine/Source/Runtime/HeadMountedDisplay/Private/HeadMountedDisplayFunctionLibrary.cpp
--- a/Engine/Source/Runtime/HeadMountedDisplay/Private/HeadMountedDisplayFunctionLibrary.cpp
+++ b/Engine/Source/Runtime/HeadMountedDisplay/Private/HeadMountedDisplayFunctionLibrary.cpp
@@ -567,3 +567,28 @@ bool UHeadMountedDisplayFunctionLibrary::GetControllerTransformForTime(UObject*
 	}
 	return false;
 }
+
+
+FVector2D UHeadMountedDisplayFunctionLibrary::GetPlayAreaBounds(TEnumAsByte<EHMDTrackingOrigin::Type> InOrigin)
+{
+	if (GEngine->XRSystem.IsValid())
+	{
+		EHMDTrackingOrigin::Type Origin = EHMDTrackingOrigin::Stage;
+		switch (InOrigin)
+		{
+		case EHMDTrackingOrigin::Eye:
+			Origin = EHMDTrackingOrigin::Eye;
+			break;
+		case EHMDTrackingOrigin::Floor:
+			Origin = EHMDTrackingOrigin::Floor;
+			break;
+		case EHMDTrackingOrigin::Stage:
+			Origin = EHMDTrackingOrigin::Stage;
+			break;
+		default:
+			break;
+		}
+		return GEngine->XRSystem->GetPlayAreaBounds(Origin);
+	}
+	return FVector2D::ZeroVector;
+}
diff --git a/Engine/Source/Runtime/HeadMountedDisplay/Public/HeadMountedDisplayFunctionLibrary.h b/Engine/Source/Runtime/HeadMountedDisplay/Public/HeadMountedDisplayFunctionLibrary.h
--- a/Engine/Source/Runtime/HeadMountedDisplay/Public/HeadMountedDisplayFunctionLibrary.h
+++ b/Engine/Source/Runtime/HeadMountedDisplay/Public/HeadMountedDisplayFunctionLibrary.h
@@ -372,4 +372,10 @@ class HEADMOUNTEDDISPLAY_API UHeadMountedDisplayFunctionLibrary : public UBluepr
 	*/
 	UFUNCTION(BlueprintCallable, Category = "Input|XRTracking")
 	static bool GetControllerTransformForTime(UObject* WorldContext, const int32 ControllerIndex, const FName MotionSource, FTimespan Time, bool& bTimeWasUsed, FRotator& Orientation, FVector& Position, bool& bProvidedLinearVelocity, FVector& LinearVelocity, bool& bProvidedAngularVelocity, FVector& AngularVelocityRadPerSec);
+
+	/**
+	 * Get the bounds of the area where the user can freely move while remaining tracked centered around the specified origin
+	 */
+	UFUNCTION(BlueprintCallable, Category = "Input|XRTracking")
+	static FVector2D GetPlayAreaBounds(TEnumAsByte<EHMDTrackingOrigin::Type> Origin = EHMDTrackingOrigin::Stage);
 };
diff --git a/Engine/Source/Runtime/HeadMountedDisplay/Public/HeadMountedDisplayTypes.h b/Engine/Source/Runtime/HeadMountedDisplay/Public/HeadMountedDisplayTypes.h
--- a/Engine/Source/Runtime/HeadMountedDisplay/Public/HeadMountedDisplayTypes.h
+++ b/Engine/Source/Runtime/HeadMountedDisplay/Public/HeadMountedDisplayTypes.h
@@ -65,8 +65,7 @@ namespace EHMDTrackingOrigin
 	{
 		Floor UMETA(DisplayName = "Floor Level"),
 		Eye UMETA(DisplayName = "Eye Level"),
-		Stage UMETA(DisplayName = "Stage (Centered Around Play Area)"),
-		Unbounded UMETA(DisplayName = "Unbounded (Centered Around Viewer)")
+		Stage UMETA(DisplayName = "Stage (Centered Around Play Area)")
 	};
 }
 
diff --git a/Engine/Source/Runtime/HeadMountedDisplay/Public/IXRTrackingSystem.h b/Engine/Source/Runtime/HeadMountedDisplay/Public/IXRTrackingSystem.h
--- a/Engine/Source/Runtime/HeadMountedDisplay/Public/IXRTrackingSystem.h
+++ b/Engine/Source/Runtime/HeadMountedDisplay/Public/IXRTrackingSystem.h
@@ -419,6 +419,11 @@ public:
 	virtual EXRDeviceConnectionResult::Type ConnectRemoteXRDevice(const FString& IpAddress, const int32 BitRate)
 	{ 
 		return EXRDeviceConnectionResult::FeatureNotSupported;
-	};
-	virtual void DisconnectRemoteXRDevice() {};
+	}
+	virtual void DisconnectRemoteXRDevice() {}
+
+	/**
+	 * Get the bounds of the area where the user can freely move while remaining tracked centered around the specified origin
+	 */
+	virtual FVector2D GetPlayAreaBounds(EHMDTrackingOrigin::Type Origin) const { return FVector2D::ZeroVector; }
 };
