UE-129231 FXMotionController data arrays are indexed by EHandKeypoint, needs better documentation
-added comments to UProperties explaining the indices are EHandKeypoint values.

diff --git a/Engine/Source/Runtime/HeadMountedDisplay/Public/HeadMountedDisplayTypes.h b/Engine/Source/Runtime/HeadMountedDisplay/Public/HeadMountedDisplayTypes.h
--- a/Engine/Source/Runtime/HeadMountedDisplay/Public/HeadMountedDisplayTypes.h
+++ b/Engine/Source/Runtime/HeadMountedDisplay/Public/HeadMountedDisplayTypes.h
@@ -345,10 +345,13 @@ struct HEADMOUNTEDDISPLAY_API FXRMotionControllerData
 	UPROPERTY(BlueprintReadOnly, Category = "XR")
 	FQuat AimRotation = FQuat(EForceInit::ForceInitToZero);
 
+	// The indices of this array are the values of EHandKeypoint (Palm, Wrist, ThumbMetacarpal, etc).
 	UPROPERTY(BlueprintReadOnly, Category = "XR")
 	TArray<FVector> HandKeyPositions;
+	// The indices of this array are the values of EHandKeypoint (Palm, Wrist, ThumbMetacarpal, etc).
 	UPROPERTY(BlueprintReadOnly, Category = "XR")
 	TArray<FQuat> HandKeyRotations;
+	// The indices of this array are the values of EHandKeypoint (Palm, Wrist, ThumbMetacarpal, etc).
 	UPROPERTY(BlueprintReadOnly, Category = "XR")
 	TArray<float> HandKeyRadii;
 
