Fixed assert when accessing IHeadMountedDisplay modular feature on the render thread

diff --git a/Engine/Source/Runtime/HeadMountedDisplay/Public/IHeadMountedDisplayModule.h b/Engine/Source/Runtime/HeadMountedDisplay/Public/IHeadMountedDisplayModule.h
--- a/Engine/Source/Runtime/HeadMountedDisplay/Public/IHeadMountedDisplayModule.h
+++ b/Engine/Source/Runtime/HeadMountedDisplay/Public/IHeadMountedDisplayModule.h
@@ -74,6 +74,8 @@ public:
 	 */
 	static inline IHeadMountedDisplayModule& Get()
 	{
+		// This function is accessible from the render and RHI threads, so lock the list of modular features before calling GetModularFeature
+		IModularFeatures::FScopedLockModularFeatureList ScopedLockModularFeatureList;
 		TArray<IHeadMountedDisplayModule*> HMDModules = IModularFeatures::Get().GetModularFeatureImplementations<IHeadMountedDisplayModule>(GetModularFeatureName());
 		HMDModules.Sort(FCompareModulePriority());
 		return *HMDModules[0];
@@ -86,6 +88,8 @@ public:
 	 */
 	static inline bool IsAvailable()
 	{
+		// This function is accessible from the render and RHI threads, so lock the list of modular features before calling GetModularFeature
+		IModularFeatures::FScopedLockModularFeatureList ScopedLockModularFeatureList;
 		return IModularFeatures::Get().IsModularFeatureAvailable(GetModularFeatureName());
 	}
 
