Remove GetScreenPercentage blueprint function deprecated since 4.19.
-Time for this to go.

diff --git a/Engine/Source/Runtime/HeadMountedDisplay/Private/HeadMountedDisplayFunctionLibrary.cpp b/Engine/Source/Runtime/HeadMountedDisplay/Private/HeadMountedDisplayFunctionLibrary.cpp
--- a/Engine/Source/Runtime/HeadMountedDisplay/Private/HeadMountedDisplayFunctionLibrary.cpp
+++ b/Engine/Source/Runtime/HeadMountedDisplay/Private/HeadMountedDisplayFunctionLibrary.cpp
@@ -220,13 +220,6 @@ void UHeadMountedDisplayFunctionLibrary::SetClippingPlanes(float Near, float Far
 	}
 }
 
-/** DEPRECATED - Use GetPixelDensity */
-float UHeadMountedDisplayFunctionLibrary::GetScreenPercentage()
-{
-	static const auto ScreenPercentageTCVar = IConsoleManager::Get().FindTConsoleVariableDataFloat(TEXT("r.ScreenPercentage"));
-	return ScreenPercentageTCVar->GetValueOnGameThread();
-}
-
 float UHeadMountedDisplayFunctionLibrary::GetPixelDensity()
 {
 	static const auto PixelDensityTCVar = IConsoleManager::Get().FindTConsoleVariableDataFloat(TEXT("vr.pixeldensity"));
diff --git a/Engine/Source/Runtime/HeadMountedDisplay/Public/HeadMountedDisplayFunctionLibrary.h b/Engine/Source/Runtime/HeadMountedDisplay/Public/HeadMountedDisplayFunctionLibrary.h
--- a/Engine/Source/Runtime/HeadMountedDisplay/Public/HeadMountedDisplayFunctionLibrary.h
+++ b/Engine/Source/Runtime/HeadMountedDisplay/Public/HeadMountedDisplayFunctionLibrary.h
@@ -162,11 +162,6 @@ class HEADMOUNTEDDISPLAY_API UHeadMountedDisplayFunctionLibrary : public UBluepr
 	UFUNCTION(BlueprintCallable, Category="Input|HeadMountedDisplay")
 	static void SetClippingPlanes(float Near, float Far);
 
-	/** DEPRECATED - Use GetPixelDensity */
-	UE_DEPRECATED(4.19, "Use GetPixelDensity instead")
-	UFUNCTION(BlueprintPure, Category = "Input|HeadMountedDisplay", meta = (DeprecatedFunction, DeprecationMessage = "Use GetPixelDensity instead"))
-	static float GetScreenPercentage();
-
 	/** 
 	 * Returns the current VR pixel density. Pixel density sets the VR render 
 	 * target texture size as a factor of recommended texture size. The recommended 
