XR LockModularFeatureList fixes.
-Thread safety of Modular Features has been improved and we were hitting ensures in the render thread update of motion controllers.
-Now we are caching the IMotionController that is used for the game thread update and only attempting to use that one for the render thread update.  We are also watching for unregistered modules to null out cached IMotionControllers as necessary (this is unlikely, but if it happened it would be very bad).

diff --git a/Engine/Source/Runtime/HeadMountedDisplay/Private/MotionControllerComponent.cpp b/Engine/Source/Runtime/HeadMountedDisplay/Private/MotionControllerComponent.cpp
--- a/Engine/Source/Runtime/HeadMountedDisplay/Private/MotionControllerComponent.cpp
+++ b/Engine/Source/Runtime/HeadMountedDisplay/Private/MotionControllerComponent.cpp
@@ -475,6 +475,55 @@ void UMotionControllerComponent::RefreshDisplayComponent(const bool bForceDestro
 	}
 }
 
+namespace UEMotionController {
+	// A scoped lock that must be explicitly locked and will unlock upon destruction if locked.
+	// Convenient if you only sometimes want to lock and the scopes are complicated.
+	class FScopeLockOptional
+	{
+	public:
+		FScopeLockOptional()
+		{
+		}
+
+		void Lock(FCriticalSection* InSynchObject)
+		{
+			SynchObject = InSynchObject;
+			SynchObject->Lock();
+		}
+
+		/** Destructor that performs a release on the synchronization object. */
+		~FScopeLockOptional()
+		{
+			Unlock();
+		}
+
+		void Unlock()
+		{
+			if (SynchObject)
+			{
+				SynchObject->Unlock();
+				SynchObject = nullptr;
+			}
+		}
+
+	private:
+		/** Copy constructor( hidden on purpose). */
+		FScopeLockOptional(const FScopeLockOptional& InScopeLock);
+
+		/** Assignment operator (hidden on purpose). */
+		FScopeLockOptional& operator=(FScopeLockOptional& InScopeLock)
+		{
+			return *this;
+		}
+
+	private:
+
+		// Holds the synchronization object to aggregate and scope manage.
+		FCriticalSection* SynchObject = nullptr;
+	};
+}
+
+
 //=============================================================================
 bool UMotionControllerComponent::PollControllerState(FVector& Position, FRotator& Orientation, float WorldToMetersScale)
 {
@@ -483,11 +532,40 @@ bool UMotionControllerComponent::PollControllerState(FVector& Position, FRotator
 		// Cache state from the game thread for use on the render thread
 		const AActor* MyOwner = GetOwner();
 		bHasAuthority = MyOwner->HasLocalNetOwner();
+
 	}
 
 	if(bHasAuthority)
 	{
-		TArray<IMotionController*> MotionControllers = IModularFeatures::Get().GetModularFeatureImplementations<IMotionController>(IMotionController::GetModularFeatureName());
+		UEMotionController::FScopeLockOptional LockOptional;
+
+		TArray<IMotionController*> MotionControllers;
+		if (IsInGameThread())
+		{
+			MotionControllers = IModularFeatures::Get().GetModularFeatureImplementations<IMotionController>(IMotionController::GetModularFeatureName());
+			{
+				FScopeLock Lock(&PolledMotionControllerMutex);
+				PolledMotionController_GameThread = nullptr;
+			}
+		}
+		else if (IsInRenderingThread())
+		{
+			LockOptional.Lock(&PolledMotionControllerMutex);
+			if (PolledMotionController_RenderThread != nullptr)
+			{
+				MotionControllers.Add(PolledMotionController_RenderThread);
+			}
+		}
+		else
+		{
+			// If we are in some other thread we can't use the game thread code, because the ModularFeature access isn't threadsafe.
+			// The render thread code might work, or not.  
+			// Let's do the fully safe locking version, and assert because this case is not expected.
+			checkNoEntry();
+			IModularFeatures::FScopedLockModularFeatureList FeatureListLock;
+			MotionControllers = IModularFeatures::Get().GetModularFeatureImplementations<IMotionController>(IMotionController::GetModularFeatureName());
+		}
+
 		for (auto MotionController : MotionControllers)
 		{
 			if (MotionController == nullptr)
@@ -503,6 +581,11 @@ bool UMotionControllerComponent::PollControllerState(FVector& Position, FRotator
 					InUseMotionController = MotionController;
 					OnMotionControllerUpdated();
 					InUseMotionController = nullptr;
+
+					{
+						FScopeLock Lock(&PolledMotionControllerMutex);
+						PolledMotionController_GameThread = MotionController;  // We only want a render thread update from the motion controller we polled on the game thread.
+					}
 				}
 				return true;
 			}
@@ -525,6 +608,20 @@ bool UMotionControllerComponent::PollControllerState(FVector& Position, FRotator
 	return false;
 }
 
+void UMotionControllerComponent::OnModularFeatureUnregistered(const FName& Type, class IModularFeature* ModularFeature)
+{
+	FScopeLock Lock(&PolledMotionControllerMutex);
+
+	if (ModularFeature == PolledMotionController_GameThread)
+	{
+		PolledMotionController_GameThread = nullptr;
+	}
+	if (ModularFeature == PolledMotionController_RenderThread)
+	{
+		PolledMotionController_RenderThread = nullptr;
+	}
+}
+
 //=============================================================================
 UMotionControllerComponent::FViewExtension::FViewExtension(const FAutoRegister& AutoRegister, UMotionControllerComponent* InMotionControllerComponent)
 	: FSceneViewExtensionBase(AutoRegister)
@@ -560,6 +657,11 @@ void UMotionControllerComponent::FViewExtension::PreRenderViewFamily_RenderThrea
 			return;
 		}
 
+		{
+			FScopeLock Lock(&MotionControllerComponent->PolledMotionControllerMutex);
+			MotionControllerComponent->PolledMotionController_RenderThread = MotionControllerComponent->PolledMotionController_GameThread;
+		}
+
 		// Find a view that is associated with this player.
 		float WorldToMetersScale = -1.0f;
 		for (const FSceneView* SceneView : InViewFamily.Views)
diff --git a/Engine/Source/Runtime/HeadMountedDisplay/Public/IHeadMountedDisplayModule.h b/Engine/Source/Runtime/HeadMountedDisplay/Public/IHeadMountedDisplayModule.h
--- a/Engine/Source/Runtime/HeadMountedDisplay/Public/IHeadMountedDisplayModule.h
+++ b/Engine/Source/Runtime/HeadMountedDisplay/Public/IHeadMountedDisplayModule.h
@@ -74,8 +74,7 @@ public:
 	 */
 	static inline IHeadMountedDisplayModule& Get()
 	{
-		// This function is accessible from the render and RHI threads, so lock the list of modular features before calling GetModularFeature
-		IModularFeatures::FScopedLockModularFeatureList ScopedLockModularFeatureList;
+		check(IsInGameThread());
 		TArray<IHeadMountedDisplayModule*> HMDModules = IModularFeatures::Get().GetModularFeatureImplementations<IHeadMountedDisplayModule>(GetModularFeatureName());
 		HMDModules.Sort(FCompareModulePriority());
 		return *HMDModules[0];
@@ -88,8 +87,7 @@ public:
 	 */
 	static inline bool IsAvailable()
 	{
-		// This function is accessible from the render and RHI threads, so lock the list of modular features before calling GetModularFeature
-		IModularFeatures::FScopedLockModularFeatureList ScopedLockModularFeatureList;
+		check(IsInGameThread());
 		return IModularFeatures::Get().IsModularFeatureAvailable(GetModularFeatureName());
 	}
 
diff --git a/Engine/Source/Runtime/HeadMountedDisplay/Public/MotionControllerComponent.h b/Engine/Source/Runtime/HeadMountedDisplay/Public/MotionControllerComponent.h
--- a/Engine/Source/Runtime/HeadMountedDisplay/Public/MotionControllerComponent.h
+++ b/Engine/Source/Runtime/HeadMountedDisplay/Public/MotionControllerComponent.h
@@ -137,6 +137,12 @@ private:
 	/** If true, the Position and Orientation args will contain the most recent controller state */
 	bool PollControllerState(FVector& Position, FRotator& Orientation, float WorldToMetersScale);
 
+	void OnModularFeatureUnregistered(const FName& Type, class IModularFeature* ModularFeature);
+	IMotionController* PolledMotionController_GameThread;
+	IMotionController* PolledMotionController_RenderThread;
+	FCriticalSection PolledMotionControllerMutex;
+
+
 	FTransform RenderThreadRelativeTransform;
 	FVector RenderThreadComponentScale;
 
