UE-147360 VR Preview w. Net Mode Play As Client doesn't function as intended
-Previous code said that the first PIE world was allowed to track the HMD.  However sometimes the first PIE world is a dedicated server.  Now the first non-dedicated-server PIE world gets HMD tracking.
-This fix is a bit ugly to fit with hotfix rules.  I think it would be better to cache this value somewhere, but didn't see anywhere that was not a public header.

diff --git a/Engine/Source/Runtime/HeadMountedDisplay/Private/IXRTrackingSystem.cpp b/Engine/Source/Runtime/HeadMountedDisplay/Private/IXRTrackingSystem.cpp
--- a/Engine/Source/Runtime/HeadMountedDisplay/Private/IXRTrackingSystem.cpp
+++ b/Engine/Source/Runtime/HeadMountedDisplay/Private/IXRTrackingSystem.cpp
@@ -28,7 +28,28 @@ bool IXRTrackingSystem::IsHeadTrackingAllowedForWorld(UWorld& World) const
 {
 #if WITH_EDITOR
 	// For VR PIE only the first instance uses the headset
-	return IsHeadTrackingAllowed() && ((World.WorldType != EWorldType::PIE) || (World.GetOutermost()->GetPIEInstanceID() == 0));
+	// This implementation is constrained by hotfix rules.  It would be better to cache this somewhere.
+
+	if (!IsHeadTrackingAllowed())
+	{
+		return false;
+	}
+
+	if (World.WorldType != EWorldType::PIE)
+	{
+		return true;
+	}
+
+	// If we are a pie instance then the first pie world that is not a dedicated server uses head tracking
+	const int32 MyPIEInstanceID = World.GetOutermost()->GetPIEInstanceID();
+	for (const FWorldContext& WorldContext : GEngine->GetWorldContexts())
+	{
+		if (WorldContext.WorldType == EWorldType::PIE && WorldContext.RunAsDedicated == false && WorldContext.World())
+		{
+			return WorldContext.World()->GetOutermost()->GetPIEInstanceID() == MyPIEInstanceID;
+		}
+	}
+	return false;
 #else
 	return IsHeadTrackingAllowed();
 #endif
