XR | Added keyword "HMD" to IsHeadMountedDisplayEnabled, and IsHeadMountedDisplayConnected

diff --git a/Engine/Source/Runtime/HeadMountedDisplay/Public/HeadMountedDisplayFunctionLibrary.h b/Engine/Source/Runtime/HeadMountedDisplay/Public/HeadMountedDisplayFunctionLibrary.h
--- a/Engine/Source/Runtime/HeadMountedDisplay/Public/HeadMountedDisplayFunctionLibrary.h
+++ b/Engine/Source/Runtime/HeadMountedDisplay/Public/HeadMountedDisplayFunctionLibrary.h
@@ -23,7 +23,7 @@ class HEADMOUNTEDDISPLAY_API UHeadMountedDisplayFunctionLibrary : public UBluepr
 	 *
 	 * @return (Boolean)  status of HMD
 	 */
-	UFUNCTION(BlueprintPure, Category="Input|HeadMountedDisplay")
+	UFUNCTION(BlueprintPure, Category="Input|HeadMountedDisplay", meta = (KeyWords = "HMD"))
 	static bool IsHeadMountedDisplayEnabled();
 
 	/**
@@ -31,7 +31,7 @@ class HEADMOUNTEDDISPLAY_API UHeadMountedDisplayFunctionLibrary : public UBluepr
 	*
 	* @return (Boolean)  status whether the HMD hardware is connected and ready to use.  It may or may not actually be in use. 
 	*/
-	UFUNCTION(BlueprintPure, Category = "Input|HeadMountedDisplay")
+	UFUNCTION(BlueprintPure, Category = "Input|HeadMountedDisplay", meta = (KeyWords = "HMD"))
 	static bool IsHeadMountedDisplayConnected();
 
 	/**
