OpenXR PixelDensity: custom quantization of width

In the scenario where we have an atlased stereo texture, we need to make sure each half
of the texture is aligned by 8 (new Strata requirement). We can do this by manually quantizing
the texture after applying pixel density.

diff --git a/Engine/Source/Runtime/HeadMountedDisplay/Private/XRRenderTargetManager.cpp b/Engine/Source/Runtime/HeadMountedDisplay/Private/XRRenderTargetManager.cpp
--- a/Engine/Source/Runtime/HeadMountedDisplay/Private/XRRenderTargetManager.cpp
+++ b/Engine/Source/Runtime/HeadMountedDisplay/Private/XRRenderTargetManager.cpp
@@ -19,8 +19,24 @@ void FXRRenderTargetManager::CalculateRenderTargetSize(const class FViewport& Vi
 		{
 			const FIntPoint IdealRenderTargetSize = HMDDevice->GetIdealRenderTargetSize();
 			const float PixelDensity = HMDDevice->GetPixelDenity();
-			InOutSizeX = FMath::CeilToInt(IdealRenderTargetSize.X * PixelDensity);
-			InOutSizeY = FMath::CeilToInt(IdealRenderTargetSize.Y * PixelDensity);
+
+			FIntPoint DensityAdjustedTargetSize =
+			{
+				FMath::CeilToInt(IdealRenderTargetSize.X * PixelDensity),
+				FMath::CeilToInt(IdealRenderTargetSize.Y * PixelDensity)
+			};
+
+			// We need a custom quantized width here because if we have an atlased texture, each half needs to be aligned.
+			// We could modify or overload QuantizeSceneBufferSize, but this is the only call point that needs to fix up 
+			// the pixel-density-adjusted width. Strata requires DivBy8, hence aligning to 16.
+			constexpr uint32 Mask16 = ~(16 - 1);
+			DensityAdjustedTargetSize.X = (DensityAdjustedTargetSize.X + (16 - 1)) & Mask16;
+
+			QuantizeSceneBufferSize(DensityAdjustedTargetSize, DensityAdjustedTargetSize);
+
+			InOutSizeX = DensityAdjustedTargetSize.X;
+			InOutSizeY = DensityAdjustedTargetSize.Y;
+
 			check(InOutSizeX != 0 && InOutSizeY != 0);
 		}
 	}
