Fix NoPCH error: using QuantizeSceneBufferSize w/o include

diff --git a/Engine/Source/Runtime/HeadMountedDisplay/Private/XRRenderTargetManager.cpp b/Engine/Source/Runtime/HeadMountedDisplay/Private/XRRenderTargetManager.cpp
--- a/Engine/Source/Runtime/HeadMountedDisplay/Private/XRRenderTargetManager.cpp
+++ b/Engine/Source/Runtime/HeadMountedDisplay/Private/XRRenderTargetManager.cpp
@@ -6,6 +6,7 @@
 #include "IHeadMountedDisplay.h"
 #include "IXRTrackingSystem.h"
 #include "Engine/Engine.h"
+#include "RenderUtils.h"
 #include "XRRenderBridge.h"
 
 void FXRRenderTargetManager::CalculateRenderTargetSize(const class FViewport& Viewport, uint32& InOutSizeX, uint32& InOutSizeY)
@@ -29,6 +30,7 @@ void FXRRenderTargetManager::CalculateRenderTargetSize(const class FViewport& Vi
 			// We need a custom quantized width here because if we have an atlased texture, each half needs to be aligned.
 			// We could modify or overload QuantizeSceneBufferSize, but this is the only call point that needs to fix up 
 			// the pixel-density-adjusted width. Strata requires DivBy8, hence aligning to 16.
+			// TODO: Would be nice if we could offload the alignment to QSBS by passing in number of atlased views
 			constexpr uint32 Mask16 = ~(16 - 1);
 			DensityAdjustedTargetSize.X = (DensityAdjustedTargetSize.X + (16 - 1)) & Mask16;
 
