Refactor MotionController to provide separate VisualizationComponent.
Upvotes from same cl in Release-5.1 stream https://p4-swarm.epicgames.net/reviews/23196426

diff --git a/Engine/Source/Editor/DetailCustomizations/Private/DetailCustomizations.cpp b/Engine/Source/Editor/DetailCustomizations/Private/DetailCustomizations.cpp
--- a/Engine/Source/Editor/DetailCustomizations/Private/DetailCustomizations.cpp
+++ b/Engine/Source/Editor/DetailCustomizations/Private/DetailCustomizations.cpp
@@ -160,6 +160,7 @@
 #include "VectorStructCustomization.h"
 #include "WindowsTargetSettingsDetails.h"
 #include "WorldSettingsDetails.h"
+#include "XRDeviceVisualizationDetails.h"
 #include "LandscapeGrassTypeDetails.h"
 
 struct FPerPlatformBool;
@@ -467,6 +468,7 @@ void FDetailCustomizationsModule::RegisterObjectCustomizations()
 	RegisterCustomClassLayout("Skeleton", FOnGetDetailCustomizationInstance::CreateStatic(&FSkeletonDetails::MakeInstance));
 
 	RegisterCustomClassLayout("MotionControllerComponent", FOnGetDetailCustomizationInstance::CreateStatic(&FMotionControllerDetails::MakeInstance));
+	RegisterCustomClassLayout("XRDeviceVisualizationComponent", FOnGetDetailCustomizationInstance::CreateStatic(&FXRDeviceVisualizationDetails::MakeInstance));
 
 	RegisterCustomClassLayout("Landscape", FOnGetDetailCustomizationInstance::CreateStatic(&FLandscapeUIDetails::MakeInstance));
 	RegisterCustomClassLayout("LandscapeProxy", FOnGetDetailCustomizationInstance::CreateStatic(&FLandscapeProxyUIDetails::MakeInstance));
diff --git a/Engine/Source/Editor/DetailCustomizations/Private/MotionControllerDetails.cpp b/Engine/Source/Editor/DetailCustomizations/Private/MotionControllerDetails.cpp
--- a/Engine/Source/Editor/DetailCustomizations/Private/MotionControllerDetails.cpp
+++ b/Engine/Source/Editor/DetailCustomizations/Private/MotionControllerDetails.cpp
@@ -42,6 +42,8 @@ class FProperty;
 
 #define LOCTEXT_NAMESPACE "MotionControllerDetails"
 
+PRAGMA_DISABLE_DEPRECATION_WARNINGS
+
 TMap< FName, TSharedPtr<FName> > FMotionControllerDetails::CustomSourceNames;
 
 TSharedRef<IDetailCustomization> FMotionControllerDetails::MakeInstance()
@@ -375,3 +377,5 @@ FText FMotionControllerDetails::GetMotionSourceValueText() const
 }
 
 #undef LOCTEXT_NAMESPACE
+
+PRAGMA_ENABLE_DEPRECATION_WARNINGS
diff --git a/Engine/Source/Editor/DetailCustomizations/Private/XRDeviceVisualizationDetails.cpp b/Engine/Source/Editor/DetailCustomizations/Private/XRDeviceVisualizationDetails.cpp
new file mode 100644
--- /dev/null
+++ b/Engine/Source/Editor/DetailCustomizations/Private/XRDeviceVisualizationDetails.cpp
@@ -0,0 +1,319 @@
+// Copyright Epic Games, Inc. All Rights Reserved.
+
+#include "XRDeviceVisualizationDetails.h"
+
+#include "DetailLayoutBuilder.h"
+#include "DetailWidgetRow.h"
+#include "Features/IModularFeatures.h" // for GetModularFeatureImplementations()
+#include "IXRSystemAssets.h"
+#include "MotionControllerComponent.h"
+#include "MotionControllerSourceCustomization.h"
+#include "Widgets/Input/SComboButton.h"
+#include "Widgets/Input/SEditableTextBox.h"
+#include "Widgets/Layout/SBorder.h"
+#include "Widgets/Text/STextBlock.h"
+#include "Widgets/Views/SListView.h"
+#include "Widgets/Views/STableRow.h"
+#include "XRDeviceVisualizationComponent.h"
+
+class FProperty;
+
+#define LOCTEXT_NAMESPACE "XRDeviceVisualizationDetails"
+
+TMap< FName, TSharedPtr<FName> > FXRDeviceVisualizationDetails::CustomSourceNames;
+
+TSharedRef<IDetailCustomization> FXRDeviceVisualizationDetails::MakeInstance()
+{
+	return MakeShareable(new FXRDeviceVisualizationDetails);
+}
+
+void FXRDeviceVisualizationDetails::CustomizeDetails(IDetailLayoutBuilder& DetailLayout) 
+{
+	DetailLayout.GetObjectsBeingCustomized(SelectedObjects);
+	IDetailCategoryBuilder& VisualizationDetails = DetailLayout.EditCategory("Visualization");
+	
+	TSharedRef<IPropertyHandle> ModelSrcProperty = DetailLayout.GetProperty(GET_MEMBER_NAME_CHECKED(UXRDeviceVisualizationComponent, DisplayModelSource));
+	if (ensure(ModelSrcProperty->IsValidHandle()))
+	{
+		XRSourceProperty = ModelSrcProperty;
+	}
+
+	TSharedRef<IPropertyHandle> CustomMeshProperty = DetailLayout.GetProperty(GET_MEMBER_NAME_CHECKED(UXRDeviceVisualizationComponent, CustomDisplayMesh));
+	TSharedRef<IPropertyHandle> CustomMaterialsProperty = DetailLayout.GetProperty(GET_MEMBER_NAME_CHECKED(UXRDeviceVisualizationComponent, DisplayMeshMaterialOverrides));
+
+	TArray<TSharedRef<IPropertyHandle>> VisualizationProperties;
+	VisualizationDetails.GetDefaultProperties(VisualizationProperties);
+
+	for (TSharedRef<IPropertyHandle> VisProp : VisualizationProperties)
+	{
+		IDetailPropertyRow& PropertyRow = VisualizationDetails.AddProperty(VisProp);
+
+		FProperty* RawProperty = VisProp->GetProperty();
+		if (RawProperty == ModelSrcProperty->GetProperty())
+		{
+			CustomizeModelSourceRow(ModelSrcProperty, PropertyRow);
+		}
+		else if (RawProperty == CustomMeshProperty->GetProperty())
+		{
+			CustomizeCustomMeshRow(PropertyRow);
+		}
+	}
+}
+
+void FXRDeviceVisualizationDetails::RefreshXRSourceList()
+{
+	TSharedPtr<FName> DefaultOption;
+	TSharedPtr<FName> CustomOption;
+	if (XRSourceNames.Num() > 0)
+	{
+		DefaultOption = XRSourceNames[0];
+		CustomOption  = XRSourceNames.Last();
+	}
+
+	TArray<IXRSystemAssets*> XRAssetSystems = IModularFeatures::Get().GetModularFeatureImplementations<IXRSystemAssets>(IXRSystemAssets::GetModularFeatureName());
+	XRSourceNames.Empty(XRAssetSystems.Num() + CustomSourceNames.Num() + 2);
+
+	if (!DefaultOption.IsValid())
+	{
+		DefaultOption = MakeShareable(new FName());
+	}
+	XRSourceNames.Add(DefaultOption);
+
+	TArray<FName> ListedNames;
+	TArray<FName> StaleCustomNames;
+	
+	for (auto& CustomNameIt : CustomSourceNames)
+	{
+		UObject* CustomizedController = FindObject<UMotionControllerComponent>(/*Outer =*/nullptr, *CustomNameIt.Key.ToString());
+		if (CustomizedController)
+		{
+			if (!CustomNameIt.Value->IsNone() && *CustomNameIt.Value != UXRDeviceVisualizationComponent::CustomModelSourceId)
+			{
+				ListedNames.AddUnique(*CustomNameIt.Value);
+				XRSourceNames.AddUnique(CustomNameIt.Value);
+			}
+		}
+		else
+		{
+			StaleCustomNames.Add(CustomNameIt.Key);
+		}
+	}
+	for (FName& DeadEntry : StaleCustomNames)
+	{
+		CustomSourceNames.Remove(DeadEntry);
+	}
+
+	for (IXRSystemAssets* AssetSys : XRAssetSystems)
+	{
+		FName SystemName = AssetSys->GetSystemName();
+		if (!ListedNames.Contains(SystemName))
+		{
+			ListedNames.Add(SystemName);
+
+			// @TODO: shouldn't be continuously creating these
+			TSharedPtr<FName> SystemNamePtr = MakeShareable(new FName(SystemName));
+			XRSourceNames.AddUnique(SystemNamePtr);
+		}
+	}
+
+	if (!CustomOption.IsValid())
+	{
+		CustomOption = MakeShareable(new FName(UXRDeviceVisualizationComponent::CustomModelSourceId));
+	}
+	XRSourceNames.AddUnique(CustomOption);
+}
+
+void FXRDeviceVisualizationDetails::SetSourcePropertyValue(const FName NewSystemName)
+{
+	if (XRSourceProperty.IsValid())
+	{
+		XRSourceProperty->SetValue(NewSystemName);
+	}	
+}
+
+void FXRDeviceVisualizationDetails::UpdateSourceSelection(TSharedPtr<FName> NewSelection)
+{
+	for (TWeakObjectPtr<UObject> SelectedObj : SelectedObjects)
+	{
+		if (SelectedObj.IsValid())
+		{
+			CustomSourceNames.Add(*SelectedObj->GetPathName(), NewSelection);
+		}
+	}
+	SetSourcePropertyValue(*NewSelection);
+}
+
+void FXRDeviceVisualizationDetails::CustomizeModelSourceRow(TSharedRef<IPropertyHandle>& Property, IDetailPropertyRow& PropertyRow)
+{
+	const FText ToolTip = Property->GetToolTipText();
+
+	FResetToDefaultOverride ResetOverride = FResetToDefaultOverride::Create(
+		FIsResetToDefaultVisible::CreateSP(this, &FXRDeviceVisualizationDetails::IsSourceValueModified),
+		FResetToDefaultHandler::CreateSP(this, &FXRDeviceVisualizationDetails::OnResetSourceValue)
+	);
+	PropertyRow.OverrideResetToDefault(ResetOverride);
+
+	PropertyRow.CustomWidget()
+		.NameContent()
+		[
+			Property->CreatePropertyNameWidget()
+		]
+		.ValueContent()
+		[
+			SNew(SComboButton)
+				.ContentPadding(FMargin(0,0,5,0))
+				.OnComboBoxOpened(this, &FXRDeviceVisualizationDetails::OnSourceMenuOpened)
+				.ButtonContent()
+				[
+					SNew(SBorder)
+						.BorderImage(FAppStyle::GetBrush("NoBorder"))
+						.Padding(FMargin(0, 0, 5, 0))
+					[
+						SNew(SEditableTextBox)
+ 							.Text(this, &FXRDeviceVisualizationDetails::OnGetSelectedSourceText)
+ 							.OnTextCommitted(this, &FXRDeviceVisualizationDetails::OnSourceNameCommited)
+ 							.ToolTipText(ToolTip)
+ 							.SelectAllTextWhenFocused(true)
+ 							.RevertTextOnEscape(true)
+ 							.Font(IDetailLayoutBuilder::GetDetailFont())
+					]
+				]
+				.MenuContent()
+				[
+					SNew(SVerticalBox)
+						+SVerticalBox::Slot()
+							.AutoHeight()
+							.MaxHeight(400.0f)
+						[
+							SNew(SListView< TSharedPtr<FName> >)
+ 								.ListItemsSource(&XRSourceNames)
+ 								.OnGenerateRow(this, &FXRDeviceVisualizationDetails::MakeSourceSelectionWidget)
+ 								.OnSelectionChanged(this, &FXRDeviceVisualizationDetails::OnSourceSelectionChanged)
+						]
+				]
+		];
+}
+
+void FXRDeviceVisualizationDetails::OnResetSourceValue(TSharedPtr<IPropertyHandle> /*PropertyHandle*/)
+{
+	TSharedPtr<FName> DefaultOption;
+	if (XRSourceNames.Num() > 0)
+	{
+		DefaultOption = XRSourceNames[0];
+	}
+	else
+	{
+		DefaultOption = MakeShareable(new FName());
+		XRSourceNames.Add(DefaultOption);
+	}
+	UpdateSourceSelection(DefaultOption);
+}
+
+bool FXRDeviceVisualizationDetails::IsSourceValueModified(TSharedPtr<IPropertyHandle> PropertyHandle)
+{
+	FName CurrentValue;
+	PropertyHandle->GetValue(CurrentValue);
+
+	return !CurrentValue.IsNone();
+}
+
+void FXRDeviceVisualizationDetails::OnSourceMenuOpened()
+{
+	RefreshXRSourceList();
+}
+
+FText FXRDeviceVisualizationDetails::OnGetSelectedSourceText() const
+{
+	FText DisplayText = LOCTEXT("DefaultModelSrc", "Default");
+	if (XRSourceProperty.IsValid())
+	{
+		FName PropertyValue;
+		XRSourceProperty->GetValue(PropertyValue);
+
+		if (PropertyValue == UXRDeviceVisualizationComponent::CustomModelSourceId)
+		{
+			DisplayText = LOCTEXT("CustomModelSrc", "Custom...");
+		}
+		else if (!PropertyValue.IsNone())
+		{
+			DisplayText = FText::FromName(PropertyValue);
+		}
+	}
+	return DisplayText;
+}
+
+void FXRDeviceVisualizationDetails::OnSourceNameCommited(const FText& NewText, ETextCommit::Type InTextCommit)
+{
+	if (InTextCommit == ETextCommit::OnEnter || InTextCommit == ETextCommit::OnUserMovedFocus)
+	{
+		const FName NewName = *NewText.ToString();
+		RefreshXRSourceList();
+
+		if (NewText.IsEmpty())
+		{
+			UpdateSourceSelection(XRSourceNames[0]);
+			return;
+		}
+		else if (NewName == UXRDeviceVisualizationComponent::CustomModelSourceId)
+		{
+			UpdateSourceSelection(XRSourceNames.Last());
+			return;
+		}
+
+		TSharedPtr<FName> SelectedName;
+		for (const TSharedPtr<FName>& SystemName : XRSourceNames)
+		{
+			if (*SystemName == NewName)
+			{
+				SelectedName = SystemName;
+			}
+		}
+
+		if (!SelectedName.IsValid())
+		{
+			SelectedName = MakeShareable(new FName(NewName));
+		}
+		UpdateSourceSelection(SelectedName);
+	}
+}
+
+TSharedRef<ITableRow> FXRDeviceVisualizationDetails::MakeSourceSelectionWidget(TSharedPtr<FName> Item, const TSharedRef<STableViewBase>& OwnerTable)
+{
+	FText DisplayText;
+	if (Item.IsValid() && !Item->IsNone())
+	{
+		DisplayText = FText::FromName(*Item);
+	}
+
+	return SNew(STableRow<TSharedPtr<FString>>, OwnerTable)
+		[
+			SNew(STextBlock)
+				.Text(DisplayText)
+		];
+}
+
+void FXRDeviceVisualizationDetails::OnSourceSelectionChanged(TSharedPtr<FName> NewSelection, ESelectInfo::Type /*SelectInfo*/)
+{
+	UpdateSourceSelection(NewSelection);
+}
+
+void FXRDeviceVisualizationDetails::CustomizeCustomMeshRow(IDetailPropertyRow& PropertyRow)
+{
+	if (!UseCustomMeshAttr.IsBound())
+	{
+		UseCustomMeshAttr.Bind(this, &FXRDeviceVisualizationDetails::IsCustomMeshPropertyEnabled);
+	}
+	PropertyRow.EditCondition(UseCustomMeshAttr, nullptr);
+}
+
+bool FXRDeviceVisualizationDetails::IsCustomMeshPropertyEnabled() const
+{
+	FName SourceSetting;
+	if (XRSourceProperty.IsValid())
+	{
+		XRSourceProperty->GetValue(SourceSetting);
+	}
+	return SourceSetting == UXRDeviceVisualizationComponent::CustomModelSourceId;
+}
+
+#undef LOCTEXT_NAMESPACE
diff --git a/Engine/Source/Editor/DetailCustomizations/Private/XRDeviceVisualizationDetails.h b/Engine/Source/Editor/DetailCustomizations/Private/XRDeviceVisualizationDetails.h
new file mode 100644
--- /dev/null
+++ b/Engine/Source/Editor/DetailCustomizations/Private/XRDeviceVisualizationDetails.h
@@ -0,0 +1,57 @@
+// Copyright Epic Games, Inc. All Rights Reserved.
+
+#pragma once
+
+#include "Containers/Array.h"
+#include "Containers/Map.h"
+#include "IDetailCustomization.h"
+#include "Misc/Attribute.h"
+#include "Templates/SharedPointer.h"
+#include "Templates/TypeHash.h"
+#include "Types/SlateEnums.h" // for ETextCommit, ESelectInfo
+#include "UObject/WeakObjectPtr.h"
+#include "UObject/WeakObjectPtrTemplates.h"
+
+class IDetailLayoutBuilder;
+class IDetailPropertyRow;
+class IPropertyHandle;
+class ITableRow;
+class STableViewBase;
+class UObject;
+
+class FXRDeviceVisualizationDetails : public IDetailCustomization
+{
+public:
+	/** Makes a new instance of this detail layout class for a specific detail view requesting it */
+	static TSharedRef<IDetailCustomization> MakeInstance();
+
+	/** IDetailCustomization interface */
+	virtual void CustomizeDetails(IDetailLayoutBuilder& DetailLayout) override;
+
+private:
+	void RefreshXRSourceList();
+	void SetSourcePropertyValue(const FName NewSystemName);
+	void UpdateSourceSelection(TSharedPtr<FName> NewSelection);
+
+	void CustomizeModelSourceRow(TSharedRef<IPropertyHandle>& Property, IDetailPropertyRow& PropertyRow);
+	void OnResetSourceValue(TSharedPtr<IPropertyHandle> PropertyHandle);
+	bool IsSourceValueModified(TSharedPtr<IPropertyHandle> PropertyHandle);
+	FText OnGetSelectedSourceText() const;
+	void OnSourceMenuOpened();
+	void OnSourceNameCommited(const FText& NewText, ETextCommit::Type InTextCommit);
+	TSharedRef<ITableRow> MakeSourceSelectionWidget(TSharedPtr<FName> Item, const TSharedRef<STableViewBase>& OwnerTable);
+	void OnSourceSelectionChanged(TSharedPtr<FName> NewSelection, ESelectInfo::Type SelectInfo);
+	
+	void CustomizeCustomMeshRow(IDetailPropertyRow& PropertyRow);
+	bool IsCustomMeshPropertyEnabled() const;
+
+	TArray< TWeakObjectPtr<UObject> > SelectedObjects;
+	TSharedPtr<IPropertyHandle> XRSourceProperty;
+	TArray< TSharedPtr<FName> > XRSourceNames;
+	TAttribute<bool> UseCustomMeshAttr;
+	TSharedPtr<IPropertyHandle> DisplayMaterialsProperty;
+
+	static TMap< FName, TSharedPtr<FName> > CustomSourceNames;
+
+	TSharedPtr<IPropertyHandle> MotionSourceProperty;
+};
diff --git a/Engine/Source/Runtime/HeadMountedDisplay/Private/MotionControllerComponent.cpp b/Engine/Source/Runtime/HeadMountedDisplay/Private/MotionControllerComponent.cpp
--- a/Engine/Source/Runtime/HeadMountedDisplay/Private/MotionControllerComponent.cpp
+++ b/Engine/Source/Runtime/HeadMountedDisplay/Private/MotionControllerComponent.cpp
@@ -1,21 +1,16 @@
 // Copyright Epic Games, Inc. All Rights Reserved.
 //
 #include "MotionControllerComponent.h"
-#include "GameFramework/Pawn.h"
 #include "PrimitiveSceneProxy.h"
 #include "Misc/ScopeLock.h"
-#include "EngineGlobals.h"
-#include "Engine/Engine.h"
 #include "Features/IModularFeatures.h"
 #include "IMotionController.h"
-#include "PrimitiveSceneInfo.h"
-#include "Engine/World.h"
 #include "GameFramework/WorldSettings.h"
 #include "IXRSystemAssets.h"
 #include "Components/StaticMeshComponent.h"
 #include "MotionDelayBuffer.h"
 #include "UObject/VRObjectVersion.h"
-#include "UObject/UObjectGlobals.h" // for FindObject<>
+#include "XRDeviceVisualizationComponent.h"
 #include "XRMotionControllerBase.h"
 #include "IXRTrackingSystem.h"
 
@@ -57,6 +52,9 @@
 		return false;
 	}
 }
+
+PRAGMA_DISABLE_DEPRECATION_WARNINGS
+
 //=============================================================================
 UMotionControllerComponent::UMotionControllerComponent(const FObjectInitializer& ObjectInitializer)
 	: Super(ObjectInitializer)
@@ -94,37 +92,42 @@
 	}
 }
 
-void UMotionControllerComponent::CreateRenderState_Concurrent(FRegisterComponentContext* Context)
+//=============================================================================
+void UMotionControllerComponent::PostLoad()
 {
-	Super::CreateRenderState_Concurrent(Context);
-	RenderThreadRelativeTransform = GetRelativeTransform();
-	RenderThreadComponentScale = GetComponentScale();
-}
-
-void UMotionControllerComponent::SendRenderTransform_Concurrent()
-{
-	struct FPrimitiveUpdateRenderThreadRelativeTransformParams
+	Super::PostLoad();
+	if (bDisplayDeviceModel)
 	{
-		FTransform RenderThreadRelativeTransform;
-		FVector RenderThreadComponentScale;
-	};
-
-	FPrimitiveUpdateRenderThreadRelativeTransformParams UpdateParams;
-	UpdateParams.RenderThreadRelativeTransform = GetRelativeTransform();
-	UpdateParams.RenderThreadComponentScale = GetComponentScale();
-
-	ENQUEUE_RENDER_COMMAND(UpdateRTRelativeTransformCommand)(
-		[UpdateParams, this](FRHICommandListImmediate& RHICmdList)
-	{
-		RenderThreadRelativeTransform = UpdateParams.RenderThreadRelativeTransform;
-		RenderThreadComponentScale = UpdateParams.RenderThreadComponentScale;
-	});
-
-	Super::SendRenderTransform_Concurrent();
+		UE_LOG(LogMotionControllerComponent, Warning, TEXT("bDisplayDeviceModel is deprecated. Please use XrDeviceVisualizationComponent instead."));
+	}
 }
 
 //=============================================================================
-void UMotionControllerComponent::TickComponent(float DeltaTime, enum ELevelTick TickType, FActorComponentTickFunction *ThisTickFunction)
+void UMotionControllerComponent::BeginPlay() 
 {
+	Super::BeginPlay();
+
+	TArray<USceneComponent*> ChildrenComponents;
+	GetChildrenComponents(/*bIncludeAllDescendants*/ false, ChildrenComponents);
+
+	for (USceneComponent* ChildComponent : ChildrenComponents) 
+	{
+		VisualizationComponent = Cast<UXRDeviceVisualizationComponent>(ChildComponent);
+		if (VisualizationComponent != nullptr) 
+		{
+			return;
+		}
+	}
+}
+
+//=============================================================================
+FName UMotionControllerComponent::GetTrackingMotionSource() 
+{
+	return MotionSource;
+}
+
+//=============================================================================
+void UMotionControllerComponent::TickComponent(float DeltaTime, enum ELevelTick TickType, FActorComponentTickFunction* ThisTickFunction)
+{
 	Super::TickComponent(DeltaTime, TickType, ThisTickFunction);
 
@@ -139,11 +142,24 @@
 			SetRelativeLocationAndRotation(Position, Orientation);
 		}
 
-		// if controller tracking just kicked in or we haven't gotten a valid model yet
+		// if controller tracking just kicked in or we haven't started rendering in the (possibly present) 
+		// visualization component.
+		if (!bTracked && bNewTrackedState && !bHasStartedRendering)
+		{
+			if (VisualizationComponent)
+			{
+				VisualizationComponent->SetIsRenderingActive(true);
+				bHasStartedRendering = true;
+			}
+		}
+
+		// This part is deprecated and will be removed in later versions.
+		// If controller tracking just kicked in or we haven't gotten a valid model yet
 		if (((!bTracked && bNewTrackedState) || !DisplayComponent) && bDisplayDeviceModel && DisplayModelSource != UMotionControllerComponent::CustomModelSourceId)
 		{
 			RefreshDisplayComponent();
-		}
+		} // End of deprecation
+
 		bTracked = bNewTrackedState;
 
 		if (!ViewExtension.IsValid() && GEngine)
@@ -154,61 +170,34 @@
 }
 
 //=============================================================================
-void UMotionControllerComponent::SetShowDeviceModel(const bool bShowDeviceModel)
+void UMotionControllerComponent::CreateRenderState_Concurrent(FRegisterComponentContext* Context)
 {
-	if (bDisplayDeviceModel != bShowDeviceModel)
-	{
-		bDisplayDeviceModel = bShowDeviceModel;
-#if WITH_EDITORONLY_DATA
-		const UWorld* MyWorld = GetWorld();
-		const bool bIsGameInst = MyWorld && MyWorld->WorldType != EWorldType::Editor && MyWorld->WorldType != EWorldType::EditorPreview;
-
-		if (!bIsGameInst)
-		{
-			// tear down and destroy the existing component if we're an editor inst
-			RefreshDisplayComponent(/*bForceDestroy =*/true);
-		}
-		else
-#endif
-		if (DisplayComponent)
-		{
-			DisplayComponent->SetHiddenInGame(!bShowDeviceModel, /*bPropagateToChildren =*/false);
-		}
-		else if (!bShowDeviceModel)
-		{
-			RefreshDisplayComponent();
-		}
-	}
+	Super::CreateRenderState_Concurrent(Context);
+	RenderThreadRelativeTransform = GetRelativeTransform();
+	RenderThreadComponentScale = GetComponentScale();
 }
 
 //=============================================================================
-void UMotionControllerComponent::SetDisplayModelSource(const FName NewDisplayModelSource)
+void UMotionControllerComponent::SendRenderTransform_Concurrent()
 {
-	if (NewDisplayModelSource != DisplayModelSource)
+	struct FPrimitiveUpdateRenderThreadRelativeTransformParams
 	{
-		DisplayModelSource = NewDisplayModelSource;
-		RefreshDisplayComponent();
-	}
+		FTransform RenderThreadRelativeTransform;
+		FVector RenderThreadComponentScale;
+	};
+
+	FPrimitiveUpdateRenderThreadRelativeTransformParams UpdateParams;
+	UpdateParams.RenderThreadRelativeTransform = GetRelativeTransform();
+	UpdateParams.RenderThreadComponentScale = GetComponentScale();
+
+	ENQUEUE_RENDER_COMMAND(UpdateRTRelativeTransformCommand)(
+		[UpdateParams, this](FRHICommandListImmediate& RHICmdList)
+	{
+		RenderThreadRelativeTransform = UpdateParams.RenderThreadRelativeTransform;
+		RenderThreadComponentScale = UpdateParams.RenderThreadComponentScale;
+	});
+
+	Super::SendRenderTransform_Concurrent();
 }
 
 //=============================================================================
-void UMotionControllerComponent::SetCustomDisplayMesh(UStaticMesh* NewDisplayMesh)
-{
-	if (NewDisplayMesh != CustomDisplayMesh)
-	{
-		CustomDisplayMesh = NewDisplayMesh;
-		if (DisplayModelSource == UMotionControllerComponent::CustomModelSourceId)
-		{
-			if (UStaticMeshComponent* AsMeshComponent = Cast<UStaticMeshComponent>(DisplayComponent))
-			{
-				AsMeshComponent->SetStaticMesh(NewDisplayMesh);
-			}
-			else
-			{
-				RefreshDisplayComponent();
-			}
-		}
-	}
-}
-
-//=============================================================================
@@ -268,48 +257,6 @@
 	}
 }
 
-#if WITH_EDITOR
 //=============================================================================
-void UMotionControllerComponent::PreEditChange(FProperty* PropertyAboutToChange)
-{
-	PreEditMaterialCount = DisplayMeshMaterialOverrides.Num();
-	Super::PreEditChange(PropertyAboutToChange);
-}
-
-//=============================================================================
-void UMotionControllerComponent::PostEditChangeProperty(FPropertyChangedEvent& PropertyChangedEvent)
-{
-	Super::PostEditChangeProperty(PropertyChangedEvent);
-
-	FProperty* PropertyThatChanged = PropertyChangedEvent.Property;
-	const FName PropertyName = (PropertyThatChanged != nullptr) ? PropertyThatChanged->GetFName() : NAME_None;
-
-	if (PropertyName == GET_MEMBER_NAME_CHECKED(UMotionControllerComponent, bDisplayDeviceModel))
-	{
-		RefreshDisplayComponent(/*bForceDestroy =*/true);
-	}
-	else if (PropertyName == GET_MEMBER_NAME_CHECKED(UMotionControllerComponent, DisplayMeshMaterialOverrides))
-	{
-		RefreshDisplayComponent(/*bForceDestroy =*/DisplayMeshMaterialOverrides.Num() < PreEditMaterialCount);
-	}
-	else if (PropertyName == GET_MEMBER_NAME_CHECKED(UMotionControllerComponent, CustomDisplayMesh))
-	{
-		RefreshDisplayComponent(/*bForceDestroy =*/false);
-	}
-}
-#endif
-
-//=============================================================================
-void UMotionControllerComponent::OnRegister()
-{
-	Super::OnRegister();
-
-	if (DisplayComponent == nullptr)
-	{
-		RefreshDisplayComponent();
-	}
-}
-
-//=============================================================================
 void UMotionControllerComponent::InitializeComponent()
 {
@@ -322,161 +269,6 @@
 	}
 }
 
-//=============================================================================
-void UMotionControllerComponent::OnComponentDestroyed(bool bDestroyingHierarchy)
-{
-	Super::OnComponentDestroyed(bDestroyingHierarchy);
-
-	if (DisplayComponent)
-	{
-		DisplayComponent->DestroyComponent();
-	}
-}
-
-//=============================================================================
-void UMotionControllerComponent::RefreshDisplayComponent(const bool bForceDestroy)
-{
-	if (IsRegistered())
-	{
-		TArray<USceneComponent*> DisplayAttachChildren;
-		auto DestroyDisplayComponent = [this, &DisplayAttachChildren]()
-		{
-			DisplayDeviceId.Clear();
-
-			if (DisplayComponent)
-			{
-				// @TODO: save/restore socket attachments as well
-				DisplayAttachChildren = DisplayComponent->GetAttachChildren();
-
-				DisplayComponent->DestroyComponent(/*bPromoteChildren =*/true);
-				DisplayComponent = nullptr;
-			}
-		};
-		if (bForceDestroy)
-		{
-			DestroyDisplayComponent();
-		}
-
-		UPrimitiveComponent* NewDisplayComponent = nullptr;
-		if (bDisplayDeviceModel)
-		{
-			const EObjectFlags SubObjFlags = RF_Transactional | RF_TextExportTransient;
-
-			if (DisplayModelSource == UMotionControllerComponent::CustomModelSourceId)
-			{
-				UStaticMeshComponent* MeshComponent = nullptr;
-				if ((DisplayComponent == nullptr) || (DisplayComponent->GetClass() != UStaticMeshComponent::StaticClass()))
-				{
-					DestroyDisplayComponent();
-
-					const FName SubObjName = MakeUniqueObjectName(this, UStaticMeshComponent::StaticClass(), TEXT("MotionControllerMesh"));
-					MeshComponent = NewObject<UStaticMeshComponent>(this, SubObjName, SubObjFlags);
-				}
-				else
-				{
-					MeshComponent = CastChecked<UStaticMeshComponent>(DisplayComponent);
-				}
-				NewDisplayComponent = MeshComponent;
-
-				if (ensure(MeshComponent))
-				{
-					if (CustomDisplayMesh)
-					{
-						MeshComponent->SetStaticMesh(CustomDisplayMesh);					
-					}
-					else
-					{
-						UE_LOG(LogMotionControllerComponent, Warning, TEXT("Failed to create a custom display component for the MotionController since no mesh was specified."));
-					}
-				}				
-			}
-			else
-			{
-				TArray<IXRSystemAssets*> XRAssetSystems = IModularFeatures::Get().GetModularFeatureImplementations<IXRSystemAssets>(IXRSystemAssets::GetModularFeatureName());				
-				for (IXRSystemAssets* AssetSys : XRAssetSystems)
-				{
-					if (!DisplayModelSource.IsNone() && AssetSys->GetSystemName() != DisplayModelSource)
-					{
-						continue;
-					}
-
-					int32 DeviceId = INDEX_NONE;
-					if (MotionSource == FXRMotionControllerBase::HMDSourceId)
-					{
-						DeviceId = IXRTrackingSystem::HMDDeviceId;
-					}
-					else
-					{
-						EControllerHand ControllerHandIndex;
-						if (!FXRMotionControllerBase::GetHandEnumForSourceName(MotionSource, ControllerHandIndex))
-						{
-							break;
-						}
-						DeviceId = AssetSys->GetDeviceId(ControllerHandIndex);
-					}
-
-					if (DisplayComponent && DisplayDeviceId.IsOwnedBy(AssetSys) && DisplayDeviceId.DeviceId == DeviceId)
-					{
-						// assume that the current DisplayComponent is the same one we'd get back, so don't recreate it
-						// @TODO: maybe we should add a IsCurrentlyRenderable(int32 DeviceId) to IXRSystemAssets to confirm this in some manner
-						break;
-					}
-
-					// needs to be set before CreateRenderComponent() since the LoadComplete callback may be triggered before it returns (for syncrounous loads)
-					DisplayModelLoadState = EModelLoadStatus::Pending;
-					FXRComponentLoadComplete LoadCompleteDelegate = FXRComponentLoadComplete::CreateUObject(this, &UMotionControllerComponent::OnDisplayModelLoaded);
-					
-					NewDisplayComponent = AssetSys->CreateRenderComponent(DeviceId, GetOwner(), SubObjFlags, /*bForceSynchronous=*/false, LoadCompleteDelegate);
-					if (NewDisplayComponent != nullptr)
-					{
-						if (DisplayModelLoadState != EModelLoadStatus::Complete)
-						{
-							DisplayModelLoadState = EModelLoadStatus::InProgress;
-						}
-						DestroyDisplayComponent();
-						DisplayDeviceId = FXRDeviceId(AssetSys, DeviceId);
-						break;
-					}
-					else
-					{
-						DisplayModelLoadState = EModelLoadStatus::Unloaded;
-					}
-				}
-			}
-
-			if (NewDisplayComponent && NewDisplayComponent != DisplayComponent)
-			{
-				NewDisplayComponent->SetupAttachment(this);
-				// force disable collision - if users wish to use collision, they can setup their own sub-component
-				NewDisplayComponent->SetCollisionEnabled(ECollisionEnabled::NoCollision);
-				NewDisplayComponent->RegisterComponent();
-
-				for (USceneComponent* Child : DisplayAttachChildren)
-				{
-					Child->SetupAttachment(NewDisplayComponent);
-				}
-
-				DisplayComponent = NewDisplayComponent;
-			}
-
-			if (DisplayComponent)
-			{
-				if (DisplayModelLoadState != EModelLoadStatus::InProgress)
-				{
-					OnDisplayModelLoaded(DisplayComponent);
-				}
-
-				DisplayComponent->SetHiddenInGame(bHiddenInGame);
-				DisplayComponent->SetVisibility(GetVisibleFlag());
-			}
-		}
-		else if (DisplayComponent)
-		{
-			DisplayComponent->SetHiddenInGame(true, /*bPropagateToChildren =*/false);
-		}
-	}
-}
-
 namespace UEMotionController {
 	// A scoped lock that must be explicitly locked and will unlock upon destruction if locked.
 	// Convenient if you only sometimes want to lock and the scopes are complicated.
@@ -525,7 +317,6 @@
 	};
 }
 
-
 //=============================================================================
 bool UMotionControllerComponent::PollControllerState(FVector& Position, FRotator& Orientation, float WorldToMetersScale)
 {
@@ -610,6 +401,7 @@
 	return false;
 }
 
+//=============================================================================
 void UMotionControllerComponent::OnModularFeatureUnregistered(const FName& Type, class IModularFeature* ModularFeature)
 {
 	FScopeLock Lock(&PolledMotionControllerMutex);
@@ -697,12 +489,14 @@
 	LateUpdate.Apply_RenderThread(InViewFamily.Scene, OldTransform, NewTransform);
 }
 
+//=============================================================================
 bool UMotionControllerComponent::FViewExtension::IsActiveThisFrame_Internal(const FSceneViewExtensionContext&) const
 {
 	check(IsInGameThread());
 	return MotionControllerComponent && !MotionControllerComponent->bDisableLowLatencyUpdate && CVarEnableMotionControllerLateUpdate.GetValueOnGameThread();
 }
 
+//=============================================================================
 float UMotionControllerComponent::GetParameterValue(FName InName, bool& bValueFound)
 {
 	if (InUseMotionController)
@@ -713,6 +507,7 @@
 	return 0.f;
 }
 
+//=============================================================================
 FVector UMotionControllerComponent::GetHandJointPosition(int jointIndex, bool& bValueFound)
 {
 	FVector outPosition;
@@ -728,7 +523,232 @@
 	}
 }
 
+// These functions are deprecated and will be removed in later versions.
+//=============================================================================
+void UMotionControllerComponent::SetShowDeviceModel(const bool bShowDeviceModel)
+{
+	if (bDisplayDeviceModel != bShowDeviceModel)
+	{
+		bDisplayDeviceModel = bShowDeviceModel;
+#if WITH_EDITORONLY_DATA
+		const UWorld* MyWorld = GetWorld();
+		const bool bIsGameInst = MyWorld && MyWorld->WorldType != EWorldType::Editor && MyWorld->WorldType != EWorldType::EditorPreview;
 
+		if (!bIsGameInst)
+		{
+			// tear down and destroy the existing component if we're an editor inst
+			RefreshDisplayComponent(/*bForceDestroy =*/true);
+		}
+		else
+#endif
+			if (DisplayComponent)
+			{
+				DisplayComponent->SetHiddenInGame(!bShowDeviceModel, /*bPropagateToChildren =*/false);
+			}
+			else if (!bShowDeviceModel)
+			{
+				RefreshDisplayComponent();
+			}
+	}
+}
+
+//=============================================================================
+void UMotionControllerComponent::SetDisplayModelSource(const FName NewDisplayModelSource)
+{
+	if (NewDisplayModelSource != DisplayModelSource)
+	{
+		DisplayModelSource = NewDisplayModelSource;
+		RefreshDisplayComponent();
+	}
+}
+
+//=============================================================================
+void UMotionControllerComponent::SetCustomDisplayMesh(UStaticMesh* NewDisplayMesh)
+{
+	if (NewDisplayMesh != CustomDisplayMesh)
+	{
+		CustomDisplayMesh = NewDisplayMesh;
+		if (DisplayModelSource == UMotionControllerComponent::CustomModelSourceId)
+		{
+			if (UStaticMeshComponent* AsMeshComponent = Cast<UStaticMeshComponent>(DisplayComponent))
+			{
+				AsMeshComponent->SetStaticMesh(NewDisplayMesh);
+			}
+			else
+			{
+				RefreshDisplayComponent();
+			}
+		}
+	}
+}
+
+//=============================================================================
+void UMotionControllerComponent::OnRegister()
+{
+	Super::OnRegister();
+
+	if (DisplayComponent == nullptr)
+	{
+		RefreshDisplayComponent();
+	}
+}
+
+//=============================================================================
+void UMotionControllerComponent::OnComponentDestroyed(bool bDestroyingHierarchy)
+{
+	Super::OnComponentDestroyed(bDestroyingHierarchy);
+
+	if (DisplayComponent)
+	{
+		DisplayComponent->DestroyComponent();
+	}
+}
+
+//=============================================================================
+void UMotionControllerComponent::RefreshDisplayComponent(const bool bForceDestroy)
+{
+	if (IsRegistered())
+	{
+		TArray<USceneComponent*> DisplayAttachChildren;
+		auto DestroyDisplayComponent = [this, &DisplayAttachChildren]()
+		{
+			DisplayDeviceId.Clear();
+
+			if (DisplayComponent)
+			{
+				// @TODO: save/restore socket attachments as well
+				DisplayAttachChildren = DisplayComponent->GetAttachChildren();
+
+				DisplayComponent->DestroyComponent(/*bPromoteChildren =*/true);
+				DisplayComponent = nullptr;
+			}
+		};
+		if (bForceDestroy)
+		{
+			DestroyDisplayComponent();
+		}
+
+		UPrimitiveComponent* NewDisplayComponent = nullptr;
+		if (bDisplayDeviceModel)
+		{
+			const EObjectFlags SubObjFlags = RF_Transactional | RF_TextExportTransient;
+
+			if (DisplayModelSource == UMotionControllerComponent::CustomModelSourceId)
+			{
+				UStaticMeshComponent* MeshComponent = nullptr;
+				if ((DisplayComponent == nullptr) || (DisplayComponent->GetClass() != UStaticMeshComponent::StaticClass()))
+				{
+					DestroyDisplayComponent();
+
+					const FName SubObjName = MakeUniqueObjectName(this, UStaticMeshComponent::StaticClass(), TEXT("MotionControllerMesh"));
+					MeshComponent = NewObject<UStaticMeshComponent>(this, SubObjName, SubObjFlags);
+				}
+				else
+				{
+					MeshComponent = CastChecked<UStaticMeshComponent>(DisplayComponent);
+				}
+				NewDisplayComponent = MeshComponent;
+
+				if (ensure(MeshComponent))
+				{
+					if (CustomDisplayMesh)
+					{
+						MeshComponent->SetStaticMesh(CustomDisplayMesh);
+					}
+					else
+					{
+						UE_LOG(LogMotionControllerComponent, Warning, TEXT("Failed to create a custom display component for the MotionController since no mesh was specified."));
+					}
+				}
+			}
+			else
+			{
+				TArray<IXRSystemAssets*> XRAssetSystems = IModularFeatures::Get().GetModularFeatureImplementations<IXRSystemAssets>(IXRSystemAssets::GetModularFeatureName());
+				for (IXRSystemAssets* AssetSys : XRAssetSystems)
+				{
+					if (!DisplayModelSource.IsNone() && AssetSys->GetSystemName() != DisplayModelSource)
+					{
+						continue;
+					}
+
+					int32 DeviceId = INDEX_NONE;
+					if (MotionSource == FXRMotionControllerBase::HMDSourceId)
+					{
+						DeviceId = IXRTrackingSystem::HMDDeviceId;
+					}
+					else
+					{
+						EControllerHand ControllerHandIndex;
+						if (!FXRMotionControllerBase::GetHandEnumForSourceName(MotionSource, ControllerHandIndex))
+						{
+							break;
+						}
+						DeviceId = AssetSys->GetDeviceId(ControllerHandIndex);
+					}
+
+					if (DisplayComponent && DisplayDeviceId.IsOwnedBy(AssetSys) && DisplayDeviceId.DeviceId == DeviceId)
+					{
+						// assume that the current DisplayComponent is the same one we'd get back, so don't recreate it
+						// @TODO: maybe we should add a IsCurrentlyRenderable(int32 DeviceId) to IXRSystemAssets to confirm this in some manner
+						break;
+					}
+
+					// needs to be set before CreateRenderComponent() since the LoadComplete callback may be triggered before it returns (for syncrounous loads)
+					DisplayModelLoadState = EModelLoadStatus::Pending;
+					FXRComponentLoadComplete LoadCompleteDelegate = FXRComponentLoadComplete::CreateUObject(this, &UMotionControllerComponent::OnDisplayModelLoaded);
+
+					NewDisplayComponent = AssetSys->CreateRenderComponent(DeviceId, GetOwner(), SubObjFlags, /*bForceSynchronous=*/false, LoadCompleteDelegate);
+					if (NewDisplayComponent != nullptr)
+					{
+						if (DisplayModelLoadState != EModelLoadStatus::Complete)
+						{
+							DisplayModelLoadState = EModelLoadStatus::InProgress;
+						}
+						DestroyDisplayComponent();
+						DisplayDeviceId = FXRDeviceId(AssetSys, DeviceId);
+						break;
+					}
+					else
+					{
+						DisplayModelLoadState = EModelLoadStatus::Unloaded;
+					}
+				}
+			}
+
+			if (NewDisplayComponent && NewDisplayComponent != DisplayComponent)
+			{
+				NewDisplayComponent->SetupAttachment(this);
+				// force disable collision - if users wish to use collision, they can setup their own sub-component
+				NewDisplayComponent->SetCollisionEnabled(ECollisionEnabled::NoCollision);
+				NewDisplayComponent->RegisterComponent();
+
+				for (USceneComponent* Child : DisplayAttachChildren)
+				{
+					Child->SetupAttachment(NewDisplayComponent);
+				}
+
+				DisplayComponent = NewDisplayComponent;
+			}
+
+			if (DisplayComponent)
+			{
+				if (DisplayModelLoadState != EModelLoadStatus::InProgress)
+				{
+					OnDisplayModelLoaded(DisplayComponent);
+				}
+
+				DisplayComponent->SetHiddenInGame(bHiddenInGame);
+				DisplayComponent->SetVisibility(GetVisibleFlag());
+			}
+		}
+		else if (DisplayComponent)
+		{
+			DisplayComponent->SetHiddenInGame(true, /*bPropagateToChildren =*/false);
+		}
+	}
+}
+
+//=============================================================================
 void UMotionControllerComponent::OnDisplayModelLoaded(UPrimitiveComponent* InDisplayComponent)
 {
 	if (InDisplayComponent == DisplayComponent || DisplayModelLoadState == EModelLoadStatus::Pending)
@@ -742,6 +762,39 @@
 			}
 		}
 		DisplayModelLoadState = EModelLoadStatus::Complete;
-	}	
+	}
 }
 
+#if WITH_EDITOR
+//=============================================================================
+void UMotionControllerComponent::PreEditChange(FProperty* PropertyAboutToChange)
+{
+	PreEditMaterialCount = DisplayMeshMaterialOverrides.Num();
+	Super::PreEditChange(PropertyAboutToChange);
+}
+
+//=============================================================================
+void UMotionControllerComponent::PostEditChangeProperty(FPropertyChangedEvent& PropertyChangedEvent)
+{
+	Super::PostEditChangeProperty(PropertyChangedEvent);
+
+	FProperty* PropertyThatChanged = PropertyChangedEvent.Property;
+	const FName PropertyName = (PropertyThatChanged != nullptr) ? PropertyThatChanged->GetFName() : NAME_None;
+
+	if (PropertyName == GET_MEMBER_NAME_CHECKED(UMotionControllerComponent, bDisplayDeviceModel))
+	{
+		RefreshDisplayComponent(/*bForceDestroy =*/true);
+	}
+	else if (PropertyName == GET_MEMBER_NAME_CHECKED(UMotionControllerComponent, DisplayMeshMaterialOverrides))
+	{
+		RefreshDisplayComponent(/*bForceDestroy =*/DisplayMeshMaterialOverrides.Num() < PreEditMaterialCount);
+	}
+	else if (PropertyName == GET_MEMBER_NAME_CHECKED(UMotionControllerComponent, CustomDisplayMesh))
+	{
+		RefreshDisplayComponent(/*bForceDestroy =*/false);
+	}
+}
+#endif
+// End of deprecation.
+
+PRAGMA_ENABLE_DEPRECATION_WARNINGS
diff --git a/Engine/Source/Runtime/HeadMountedDisplay/Private/XRDeviceVisualizationComponent.cpp b/Engine/Source/Runtime/HeadMountedDisplay/Private/XRDeviceVisualizationComponent.cpp
new file mode 100644
--- /dev/null
+++ b/Engine/Source/Runtime/HeadMountedDisplay/Private/XRDeviceVisualizationComponent.cpp
@@ -0,0 +1,227 @@
+// Copyright Epic Games, Inc. All Rights Reserved.
+//
+#include "XRDeviceVisualizationComponent.h"
+#include "IXRSystemAssets.h"
+#include "Features/IModularFeatures.h"
+#include "XRMotionControllerBase.h"
+#include "IXRTrackingSystem.h"
+#include "Misc/Attribute.h"
+#include "MotionControllerComponent.h"
+
+#include UE_INLINE_GENERATED_CPP_BY_NAME(XRDeviceVisualizationComponent)
+
+DEFINE_LOG_CATEGORY_STATIC(LogXRDeviceVisualizationComponent, Log, All);
+
+FName UXRDeviceVisualizationComponent::CustomModelSourceId(TEXT("Custom"));
+
+//=============================================================================
+UXRDeviceVisualizationComponent::UXRDeviceVisualizationComponent(const FObjectInitializer& ObjectInitializer)
+	: Super(ObjectInitializer)
+{
+	PrimaryComponentTick.bCanEverTick = false;
+	MotionSource = FXRMotionControllerBase::LeftHandSourceId;
+	bAutoActivate = true;
+
+	// ensure InitializeComponent() gets called
+	bWantsInitializeComponent = true;
+}
+
+//=============================================================================
+void UXRDeviceVisualizationComponent::BeginPlay() {
+	Super::BeginPlay();
+
+	UMotionControllerComponent* ParentController = FindParentMotionController();
+	if (ParentController != nullptr) 
+	{
+		MotionSource = ParentController->GetTrackingMotionSource();
+		// Wait for controller to activate the rendering for this component.
+		bIsRenderingActive = false;
+	}
+	else 
+	{
+		UE_LOG(LogXRDeviceVisualizationComponent, Error, TEXT("No parent UMotionControllerComponent for %s was found in the hierarchy. This component will not render anything."), *GetName());
+	}
+	RefreshMesh();
+}
+
+//=============================================================================
+UMotionControllerComponent* UXRDeviceVisualizationComponent::FindParentMotionController()
+{	
+	USceneComponent* Parent = GetAttachParent();
+	UMotionControllerComponent* ParentController;
+	while (Parent != nullptr) {
+		ParentController = Cast<UMotionControllerComponent>(Parent);
+		if (ParentController != nullptr) {	
+			return ParentController;
+		}
+		Parent = Parent->GetAttachParent();
+	}
+	return nullptr;
+}
+
+//=============================================================================
+bool UXRDeviceVisualizationComponent::CanDeviceBeDisplayed() 
+{
+	return bIsRenderingActive;
+}
+
+//=============================================================================
+void UXRDeviceVisualizationComponent::SetDisplayModelSource(const FName NewDisplayModelSource)
+{
+	if (NewDisplayModelSource != DisplayModelSource)
+	{
+		DisplayModelSource = NewDisplayModelSource;
+		RefreshMesh();
+	}
+}
+
+//=============================================================================
+void UXRDeviceVisualizationComponent::SetCustomDisplayMesh(UStaticMesh* NewDisplayMesh)
+{
+	if (NewDisplayMesh != CustomDisplayMesh)
+	{
+		CustomDisplayMesh = NewDisplayMesh;
+		if (DisplayModelSource == UXRDeviceVisualizationComponent::CustomModelSourceId && CustomDisplayMesh != nullptr)
+		{
+			SetStaticMesh(CustomDisplayMesh);
+		}
+	}
+}
+
+#if WITH_EDITOR
+//=============================================================================
+void UXRDeviceVisualizationComponent::PreEditChange(FProperty* PropertyAboutToChange)
+{
+	PreEditMaterialCount = DisplayMeshMaterialOverrides.Num();
+	Super::PreEditChange(PropertyAboutToChange);
+}
+#endif
+
+//=============================================================================
+void UXRDeviceVisualizationComponent::RefreshMesh()
+{
+	if (CanDeviceBeDisplayed())
+	{
+		UStaticMeshComponent* StaticMeshComponent;
+
+		if (DisplayModelSource == UXRDeviceVisualizationComponent::CustomModelSourceId)
+		{
+			if (CustomDisplayMesh)
+			{
+				SetStaticMesh(CustomDisplayMesh);
+				int32 MatCount = FMath::Min(GetNumMaterials(), DisplayMeshMaterialOverrides.Num());
+				SetMaterials(MatCount);
+			}
+			else
+			{
+				SetStaticMesh(nullptr);
+				UE_LOG(LogXRDeviceVisualizationComponent, Error, TEXT("Failed to set a custom mesh for %s because no custom mesh was specified."), *GetName());
+			}
+		}
+		else
+		{
+			TArray<IXRSystemAssets*> XRAssetSystems = IModularFeatures::Get().GetModularFeatureImplementations<IXRSystemAssets>(IXRSystemAssets::GetModularFeatureName());
+			for (IXRSystemAssets* AssetSys : XRAssetSystems)
+			{
+				if (!DisplayModelSource.IsNone() && AssetSys->GetSystemName() != DisplayModelSource)
+				{
+					continue;
+				}
+
+				int32 DeviceId = INDEX_NONE;
+				if (MotionSource == FXRMotionControllerBase::HMDSourceId)
+				{
+					DeviceId = IXRTrackingSystem::HMDDeviceId;
+				}
+				else
+				{
+					EControllerHand ControllerHandIndex;
+					if (!FXRMotionControllerBase::GetHandEnumForSourceName(MotionSource, ControllerHandIndex))
+					{
+						break;
+					}
+					DeviceId = AssetSys->GetDeviceId(ControllerHandIndex);
+				}
+
+				if (GetStaticMesh() && DisplayDeviceId.IsOwnedBy(AssetSys) && DisplayDeviceId.DeviceId == DeviceId)
+				{
+					// Assume that the current StaticMesh is the same one we'd get back, so don't recreate it
+					break;
+				}
+
+				// Needs to be set before CreateRenderComponent() since the LoadComplete callback may be triggered before it returns (for syncrounous loads)
+				DisplayModelLoadState = EModelLoadStatus::Pending;
+				FXRComponentLoadComplete LoadCompleteDelegate = FXRComponentLoadComplete::CreateUObject(this, &UXRDeviceVisualizationComponent::OnDisplayModelLoaded);
+
+				const EObjectFlags SubObjFlags = RF_Transactional | RF_TextExportTransient;
+				UPrimitiveComponent* RenderComponent = AssetSys->CreateRenderComponent(DeviceId, GetOwner(), SubObjFlags, /*bForceSynchronous=*/false, LoadCompleteDelegate);
+				
+				if (RenderComponent != nullptr) 
+				{
+					StaticMeshComponent = Cast<UStaticMeshComponent>(RenderComponent);
+					if (StaticMeshComponent != nullptr) 
+					{
+						if (DisplayModelLoadState != EModelLoadStatus::Complete)
+						{
+							DisplayModelLoadState = EModelLoadStatus::InProgress;
+						}
+						DisplayDeviceId = FXRDeviceId(AssetSys, DeviceId);
+						break;
+					}
+					else 
+					{
+						DisplayModelLoadState = EModelLoadStatus::Unloaded;
+						SetStaticMesh(nullptr);
+						UE_LOG(LogXRDeviceVisualizationComponent, Error, TEXT("CreateRenderComponent must return a UStaticMeshComponent. Setting a NULL StaticMesh."));
+					}
+				}
+				else 
+				{
+					DisplayModelLoadState = EModelLoadStatus::Unloaded;
+					SetStaticMesh(nullptr);
+					UE_LOG(LogXRDeviceVisualizationComponent, Error, TEXT("CreateRenderComponent returned a NULL UPrimitiveComponent*. Setting a NULL StaticMesh."));
+				}
+			}
+		}
+	}
+	else
+	{
+		SetStaticMesh(nullptr);
+	}
+}
+
+void UXRDeviceVisualizationComponent::SetMaterials(int32 MatCount) {
+	for (int32 MatIndex = 0; MatIndex < MatCount; ++MatIndex)
+	{
+		this->SetMaterial(MatIndex, DisplayMeshMaterialOverrides[MatIndex]);
+	}
+}
+
+void UXRDeviceVisualizationComponent::OnDisplayModelLoaded(UPrimitiveComponent* InDisplayComponent)
+{
+	if (DisplayModelLoadState == EModelLoadStatus::Pending || DisplayModelLoadState == EModelLoadStatus::InProgress)
+	{
+		UStaticMeshComponent* StaticMeshComponent = Cast<UStaticMeshComponent>(InDisplayComponent);
+		if (StaticMeshComponent != nullptr && StaticMeshComponent->GetStaticMesh()) 
+		{
+			SetStaticMesh(StaticMeshComponent->GetStaticMesh());
+			const int32 MatCount = FMath::Min(InDisplayComponent->GetNumMaterials(), DisplayMeshMaterialOverrides.Num());
+			SetMaterials(MatCount);
+			DisplayModelLoadState = EModelLoadStatus::Complete;
+		}
+		else 
+		{
+			SetStaticMesh(nullptr);
+			UE_LOG(LogXRDeviceVisualizationComponent, Warning, TEXT("OnDisplayModelLoaded expects to received a UStaticMeshComponent with a valid StaticMesh. Setting a NULL StaticMesh."));
+		}
+	}
+}
+
+void UXRDeviceVisualizationComponent::SetIsRenderingActive(bool bRenderingIsActive) 
+{
+	if (bIsRenderingActive != bRenderingIsActive) 
+	{
+		bIsRenderingActive = bRenderingIsActive;
+		RefreshMesh();
+	}
+}
\ No newline at end of file
diff --git a/Engine/Source/Runtime/HeadMountedDisplay/Public/MotionControllerComponent.h b/Engine/Source/Runtime/HeadMountedDisplay/Public/MotionControllerComponent.h
--- a/Engine/Source/Runtime/HeadMountedDisplay/Public/MotionControllerComponent.h
+++ b/Engine/Source/Runtime/HeadMountedDisplay/Public/MotionControllerComponent.h
@@ -2,8 +2,6 @@
 
 #pragma once
 
-#include "CoreMinimal.h"
-#include "UObject/ObjectMacros.h"
 #include "Components/PrimitiveComponent.h"
 #include "SceneViewExtension.h"
 #include "IMotionController.h"
@@ -15,14 +13,13 @@
 class FRHICommandListImmediate;
 class FSceneView;
 class FSceneViewFamily;
+class UXRDeviceVisualizationComponent;
 
 UCLASS(Blueprintable, meta = (BlueprintSpawnableComponent), ClassGroup = MotionController)
 class HEADMOUNTEDDISPLAY_API UMotionControllerComponent : public UPrimitiveComponent
 {
 	GENERATED_UCLASS_BODY()
 
-	void BeginDestroy() override;
-
 	/** Which player index this motion controller should automatically follow */
 	UPROPERTY(EditAnywhere, BlueprintReadWrite, BlueprintSetter = SetAssociatedPlayerIndex, Category = "MotionController")
 	int32 PlayerIndex;
@@ -36,47 +33,56 @@
 
 	/** If false, render transforms within the motion controller hierarchy will be updated a second time immediately before rendering. */
 	UPROPERTY(EditAnywhere, BlueprintReadWrite, Category = "MotionController")
-	uint32 bDisableLowLatencyUpdate:1;
+	uint32 bDisableLowLatencyUpdate : 1;
 
 	/** The tracking status for the device (e.g. full tracking, inertial tracking only, no tracking) */
 	UPROPERTY(BlueprintReadOnly, Category = "MotionController")
 	ETrackingStatus CurrentTrackingStatus;
 
-	void TickComponent(float DeltaTime, enum ELevelTick TickType, FActorComponentTickFunction *ThisTickFunction) override;
+	/** Used to visualize this component's device */
+	UXRDeviceVisualizationComponent* VisualizationComponent;
 
-	/** Whether or not this component had a valid tracked device this frame */
-	UFUNCTION(BlueprintPure, Category = "MotionController")
-	bool IsTracked() const
-	{
-		return bTracked;
-	}
-
 	/** Used to automatically render a model associated with the set hand. */
-	UPROPERTY(EditAnywhere, BlueprintReadWrite, BlueprintSetter=SetShowDeviceModel, Category="Visualization")
+	UE_DEPRECATED(5.2, "bDisplayDeviceModel is deprecated. Please use the XRDeviceVisualizationComponent for rendering instead.")
+	UPROPERTY(EditAnywhere, BlueprintReadWrite, BlueprintSetter = SetShowDeviceModel, Category = "Visualization", meta = (DeprecatedProperty, DeprecationMessage = "bDisplayDeviceModel is deprecated. Please use the XRDeviceVisualizationComponent for rendering instead."))
 	bool bDisplayDeviceModel;
 
-	UFUNCTION(BlueprintSetter)
-	void SetShowDeviceModel(const bool bShowControllerModel);
-
 	/** Determines the source of the desired model. By default, the active XR system(s) will be queried and (if available) will provide a model for the associated device. NOTE: this may fail if there's no default model; use 'Custom' to specify your own. */
-	UPROPERTY(EditAnywhere, BlueprintReadWrite, BlueprintSetter=SetDisplayModelSource, Category="Visualization", meta=(editcondition="bDisplayDeviceModel"))
+	UE_DEPRECATED(5.2, "DisplayModelSource is deprecated. Please use the XRDeviceVisualizationComponent for rendering instead.")
+	UPROPERTY(EditAnywhere, BlueprintReadWrite, BlueprintSetter = SetDisplayModelSource, Category = "Visualization", meta = (editcondition = "bDisplayDeviceModel", DeprecatedProperty, DeprecationMessage = "DisplayModelSource is deprecated. Please use the XRDeviceVisualizationComponent for rendering instead."))
 	FName DisplayModelSource;
 
 	static FName CustomModelSourceId;
-	UFUNCTION(BlueprintSetter)
-	void SetDisplayModelSource(const FName NewDisplayModelSource);
 
 	/** A mesh override that'll be displayed attached to this MotionController. */
-	UPROPERTY(EditAnywhere, BlueprintReadWrite, BlueprintSetter=SetCustomDisplayMesh, Category="Visualization", meta=(editcondition="bDisplayDeviceModel"))
+	UE_DEPRECATED(5.2, "CustomDisplayMesh is deprecated. Please use the XRDeviceVisualizationComponent for rendering instead.")
+	UPROPERTY(EditAnywhere, BlueprintReadWrite, BlueprintSetter = SetCustomDisplayMesh, Category = "Visualization", meta = (editcondition = "bDisplayDeviceModel", DeprecatedProperty, DeprecationMessage = "CustomDisplayMesh is deprecated. Please use the XRDeviceVisualizationComponent for rendering instead."))
 	TObjectPtr<UStaticMesh> CustomDisplayMesh;
 
-	UFUNCTION(BlueprintSetter)
-	void SetCustomDisplayMesh(UStaticMesh* NewDisplayMesh);
-
 	/** Material overrides for the specified display mesh. */
-	UPROPERTY(EditAnywhere, BlueprintReadOnly, Category="Visualization", meta=(editcondition="bDisplayDeviceModel"))
+	UE_DEPRECATED(5.2, "DisplayMeshMaterialOverrides is deprecated. Please use the XRDeviceVisualizationComponent for rendering instead.")
+	UPROPERTY(EditAnywhere, BlueprintReadOnly, Category = "Visualization", meta = (editcondition = "bDisplayDeviceModel", DeprecatedProperty, DeprecationMessage = "DisplayMeshMaterialOverrides is deprecated. Please use the XRDeviceVisualizationComponent for rendering instead."))
 	TArray<TObjectPtr<UMaterialInterface>> DisplayMeshMaterialOverrides;
 
+	UE_DEPRECATED(5.2, "SetShowDeviceModel is deprecated. Please use the XRDeviceVisualizationComponent for rendering instead.")
+	UFUNCTION(BlueprintSetter, meta = (DeprecatedFunction, DeprecationMessage = "SetShowDeviceModel is deprecated. Please use the XRDeviceVisualizationComponent for rendering instead."))
+	void SetShowDeviceModel(const bool bShowControllerModel);
+
+	UE_DEPRECATED(5.2, "SetDisplayModelSource is deprecated. Please use the XRDeviceVisualizationComponent for rendering instead.")
+	UFUNCTION(BlueprintSetter, meta = (DeprecatedFunction, DeprecationMessage = "SetDisplayModelSource is deprecated. Please use the XRDeviceVisualizationComponent for rendering instead."))
+	void SetDisplayModelSource(const FName NewDisplayModelSource);
+
+	UE_DEPRECATED(5.2, "SetCustomDisplayMesh is deprecated. Please use the XRDeviceVisualizationComponent for rendering instead.")
+	UFUNCTION(BlueprintSetter, meta = (DeprecatedFunction, DeprecationMessage = "SetCustomDisplayMesh is deprecated. Please use the XRDeviceVisualizationComponent for rendering instead."))
+	void SetCustomDisplayMesh(UStaticMesh* NewDisplayMesh);
+
+	/** Whether or not this component had a valid tracked device this frame */
+	UFUNCTION(BlueprintPure, Category = "MotionController")
+	bool IsTracked() const
+	{
+		return bTracked;
+	}
+
 	UFUNCTION(BlueprintSetter, meta = (DeprecatedFunction, DeprecationMessage = "Please use the Motion Source property instead of Hand"))
 	void SetTrackingSource(const EControllerHand NewSource);
 
@@ -86,9 +92,40 @@
 	UFUNCTION(BlueprintSetter)
 	void SetTrackingMotionSource(const FName NewSource);
 
+	FName GetTrackingMotionSource();
+
 	UFUNCTION(BlueprintSetter)
 	void SetAssociatedPlayerIndex(const int32 NewPlayer);
 
+	void BeginPlay() override;
+	void TickComponent(float DeltaTime, enum ELevelTick TickType, FActorComponentTickFunction* ThisTickFunction) override;
+	void BeginDestroy() override;
+
+	// The following private properties/members are now deprecated and will be removed in later versions.
+	void RefreshDisplayComponent(const bool bForceDestroy = false);
+	void PostLoad() override;
+
+	/** Callback for asynchronous display model loads (to set materials, etc.) */
+	void OnDisplayModelLoaded(UPrimitiveComponent* DisplayComponent);
+
+	UPROPERTY(Transient, BlueprintReadOnly, Category = Visualization, meta = (AllowPrivateAccess = "true"))
+	TObjectPtr<UPrimitiveComponent> DisplayComponent;
+
+	enum class EModelLoadStatus : uint8
+	{
+		Unloaded,
+		Pending,
+		InProgress,
+		Complete
+	};
+	EModelLoadStatus DisplayModelLoadState = EModelLoadStatus::Unloaded;
+
+	FXRDeviceId DisplayDeviceId;
+
+#if WITH_EDITOR
+	int32 PreEditMaterialCount = 0;
+#endif
+
 public:
 	//~ UObject interface
 	virtual void Serialize(FArchive& Ar) override;
@@ -98,7 +135,6 @@
 	virtual void PostEditChangeProperty(FPropertyChangedEvent& PropertyChangedEvent) override;
 #endif 
 
-public:
 	//~ UActorComponent interface
 	virtual void OnRegister() override;
 	virtual void InitializeComponent() override;
@@ -110,12 +146,10 @@
 	virtual void SendRenderTransform_Concurrent() override;
 	//~ End UActorComponent Interface.
 
-	void RefreshDisplayComponent(const bool bForceDestroy = false);
-
 	// Cached Motion Controller that can be read by GetParameterValue. Only valid for the duration of OnMotionControllerUpdated
 	IMotionController* InUseMotionController;
 
-	/** Blueprint Implementable function for reponding to updated data from a motion controller (so we can use custom paramater values from it) */
+	/** Blueprint Implementable function for responding to updated data from a motion controller (so we can use custom parameter values from it) */
 	UFUNCTION(BlueprintImplementableEvent, Category = "Motion Controller Update")
 	void OnMotionControllerUpdated();
 
@@ -134,6 +168,9 @@
 	/** Whether or not this component has authority within the frame*/
 	bool bHasAuthority;
 
+	/** Whether or not this component has informed the visualization component (if present) to start rendering */
+	bool bHasStartedRendering;
+
 	/** If true, the Position and Orientation args will contain the most recent controller state */
 	bool PollControllerState(FVector& Position, FRotator& Orientation, float WorldToMetersScale);
 
@@ -170,24 +207,4 @@
 		FLateUpdateManager LateUpdate;
 	};
 	TSharedPtr< FViewExtension, ESPMode::ThreadSafe > ViewExtension;	
- 
-	UPROPERTY(Transient, BlueprintReadOnly, Category=Visualization, meta=(AllowPrivateAccess="true"))
-	TObjectPtr<UPrimitiveComponent> DisplayComponent;
-
-	/** Callback for asynchronous display model loads (to set materials, etc.) */
-	void OnDisplayModelLoaded(UPrimitiveComponent* DisplayComponent);
-
-	enum class EModelLoadStatus : uint8
-	{
-		Unloaded,
-		Pending,
-		InProgress,
-		Complete
-	};
-	EModelLoadStatus DisplayModelLoadState = EModelLoadStatus::Unloaded;
-
-	FXRDeviceId DisplayDeviceId;
-#if WITH_EDITOR
-	int32 PreEditMaterialCount = 0;
-#endif
 };
diff --git a/Engine/Source/Runtime/HeadMountedDisplay/Public/XRDeviceVisualizationComponent.h b/Engine/Source/Runtime/HeadMountedDisplay/Public/XRDeviceVisualizationComponent.h
new file mode 100644
--- /dev/null
+++ b/Engine/Source/Runtime/HeadMountedDisplay/Public/XRDeviceVisualizationComponent.h
@@ -0,0 +1,77 @@
+// Copyright Epic Games, Inc. All Rights Reserved.
+
+#pragma once
+
+#include "CoreMinimal.h"
+#include "Components/StaticMeshComponent.h"
+#include "IIdentifiableXRDevice.h" // for FXRDeviceId
+#include "UObject/ObjectMacros.h"
+#include "XRDeviceVisualizationComponent.generated.h"
+
+class UMaterialInterface;
+class UMotionControllerComponent;
+class UStaticMesh;
+
+UCLASS(Blueprintable, meta = (BlueprintSpawnableComponent), ClassGroup = MotionController)
+class HEADMOUNTEDDISPLAY_API UXRDeviceVisualizationComponent : public UStaticMeshComponent
+{
+	GENERATED_UCLASS_BODY()
+
+	/** Determines the source of the desired model. By default, the active XR system(s) will be queried and (if available) will provide a model for the associated device. NOTE: this may fail if there's no default model; use 'Custom' to specify your own. */
+	UPROPERTY(EditAnywhere, BlueprintReadWrite, BlueprintSetter = SetDisplayModelSource, Category = "Visualization")
+	FName DisplayModelSource; 
+
+	static FName CustomModelSourceId;
+	UFUNCTION(BlueprintSetter)
+	void SetDisplayModelSource(const FName NewDisplayModelSource); 
+
+	/** A mesh override that'll be displayed attached to this MotionController. */
+	UPROPERTY(EditAnywhere, BlueprintReadWrite, BlueprintSetter = SetCustomDisplayMesh, Category = "Visualization")
+	TObjectPtr<UStaticMesh> CustomDisplayMesh; 
+
+	UFUNCTION(BlueprintSetter)
+	void SetCustomDisplayMesh(UStaticMesh* NewDisplayMesh); 
+
+	/** Material overrides for the specified display mesh. */
+	UPROPERTY(EditAnywhere, BlueprintReadOnly, Category = "Visualization")
+	TArray<TObjectPtr<UMaterialInterface>> DisplayMeshMaterialOverrides; 
+
+	/** Callback for asynchronous display model loads (to set materials, etc.) */
+	void OnDisplayModelLoaded(UPrimitiveComponent* DisplayComponent); 
+
+	/** Set by the parent MotionController once tracking kicks in. */
+	bool bIsRenderingActive;
+
+	/** Whether this component can be displayed, depending on whether the parent MotionController has activated its rendering. */
+	bool CanDeviceBeDisplayed();
+
+	FName MotionSource;
+
+	enum class EModelLoadStatus : uint8
+	{
+		Unloaded,
+		Pending,
+		InProgress,
+		Complete
+	};
+	EModelLoadStatus DisplayModelLoadState = EModelLoadStatus::Unloaded;
+
+	FXRDeviceId DisplayDeviceId;
+
+	void BeginPlay() override;
+	void RefreshMesh();
+	UMotionControllerComponent* FindParentMotionController();
+	void SetMaterials(int32 MatCount);
+
+#if WITH_EDITOR
+	int32 PreEditMaterialCount = 0;
+#endif
+
+public:
+#if WITH_EDITOR 
+	virtual void PreEditChange(FProperty* PropertyAboutToChange) override;
+#endif 
+
+	UFUNCTION(BlueprintSetter, Category="MotionController")
+	void SetIsRenderingActive(bool bRenderingIsActive);
+};
\ No newline at end of file
