PR #8201: Fix SetShowDeviceModelfunction bug (Contributed by [REDACTED_PERSON])

diff --git a/Engine/Source/Runtime/HeadMountedDisplay/Private/MotionControllerComponent.cpp b/Engine/Source/Runtime/HeadMountedDisplay/Private/MotionControllerComponent.cpp
--- a/Engine/Source/Runtime/HeadMountedDisplay/Private/MotionControllerComponent.cpp
+++ b/Engine/Source/Runtime/HeadMountedDisplay/Private/MotionControllerComponent.cpp
@@ -541,14 +541,14 @@ void UMotionControllerComponent::SetShowDeviceModel(const bool bShowDeviceModel)
 		}
 		else
 #endif
-			if (DisplayComponent)
-			{
-				DisplayComponent->SetHiddenInGame(!bShowDeviceModel, /*bPropagateToChildren =*/false);
-			}
-			else if (!bShowDeviceModel)
-			{
-				RefreshDisplayComponent();
-			}
+		if (DisplayComponent)
+		{
+			DisplayComponent->SetHiddenInGame(!bShowDeviceModel, /*bPropagateToChildren =*/false);
+		}
+		else if (bShowDeviceModel)
+		{
+			RefreshDisplayComponent();
+		}
 	}
 }
 
