Inserted XRDeviceVisualizationComponent in VRTemplate.
Also added Editor property and API to disable the visualization provided by the component (for example if the hand model doesn't align with the controller) replacing the old UMotionControllerComponent::SetShowDeviceModel API.

diff --git a/Engine/Source/Runtime/HeadMountedDisplay/Private/XRDeviceVisualizationComponent.cpp b/Engine/Source/Runtime/HeadMountedDisplay/Private/XRDeviceVisualizationComponent.cpp
--- a/Engine/Source/Runtime/HeadMountedDisplay/Private/XRDeviceVisualizationComponent.cpp
+++ b/Engine/Source/Runtime/HeadMountedDisplay/Private/XRDeviceVisualizationComponent.cpp
@@ -59,10 +59,20 @@ UMotionControllerComponent* UXRDeviceVisualizationComponent::FindParentMotionCon
 	return nullptr;
 }
 
+//=============================================================================
+void UXRDeviceVisualizationComponent::SetIsVisualizationActive(bool bNewVisualizationState)
+{
+	if (bNewVisualizationState != bIsVisualizationActive) 
+	{
+		bIsVisualizationActive = bNewVisualizationState;
+		RefreshMesh();
+	}
+}
+
 //=============================================================================
 bool UXRDeviceVisualizationComponent::CanDeviceBeDisplayed() 
 {
-	return bIsRenderingActive;
+	return bIsRenderingActive && bIsVisualizationActive;
 }
 
 //=============================================================================
diff --git a/Engine/Source/Runtime/HeadMountedDisplay/Public/XRDeviceVisualizationComponent.h b/Engine/Source/Runtime/HeadMountedDisplay/Public/XRDeviceVisualizationComponent.h
--- a/Engine/Source/Runtime/HeadMountedDisplay/Public/XRDeviceVisualizationComponent.h
+++ b/Engine/Source/Runtime/HeadMountedDisplay/Public/XRDeviceVisualizationComponent.h
@@ -17,6 +17,13 @@ class HEADMOUNTEDDISPLAY_API UXRDeviceVisualizationComponent : public UStaticMes
 {
 	GENERATED_UCLASS_BODY()
 
+	/** Whether the visualization offered by this component is being used. */
+	UPROPERTY(EditAnywhere, BlueprintReadWrite, BlueprintSetter = SetIsVisualizationActive, Category = "Visualization")
+	bool bIsVisualizationActive;
+
+	UFUNCTION(BlueprintSetter)
+	void SetIsVisualizationActive(bool bNewVisualizationState);
+
 	/** Determines the source of the desired model. By default, the active XR system(s) will be queried and (if available) will provide a model for the associated device. NOTE: this may fail if there's no default model; use 'Custom' to specify your own. */
 	UPROPERTY(EditAnywhere, BlueprintReadWrite, BlueprintSetter = SetDisplayModelSource, Category = "Visualization")
 	FName DisplayModelSource; 
