LWC truncation warning fixed for the HeadMountedDisplay module.

diff --git a/Engine/Source/Runtime/HeadMountedDisplay/Private/DefaultSpectatorScreenController.cpp b/Engine/Source/Runtime/HeadMountedDisplay/Private/DefaultSpectatorScreenController.cpp
--- a/Engine/Source/Runtime/HeadMountedDisplay/Private/DefaultSpectatorScreenController.cpp
+++ b/Engine/Source/Runtime/HeadMountedDisplay/Private/DefaultSpectatorScreenController.cpp
@@ -309,7 +309,7 @@ void FDefaultSpectatorScreenController::RenderSpectatorModeSingleEyeCroppedToFil
 {
 	const FIntRect SrcRect = GetFullFlatEyeRect_RenderThread(EyeTexture);
 	const FIntRect DstRect(0, 0, TargetTexture->GetSizeX(), TargetTexture->GetSizeY());
-	const FIntRect WindowRect(0, 0, WindowSize.X, WindowSize.Y);
+	const FIntRect WindowRect(0, 0, static_cast<int32>(WindowSize.X), static_cast<int32>(WindowSize.Y));
 
 	const FIntRect SrcCroppedToFitRect = Helpers::GetEyeCroppedToFitRect(HMDDevice->GetEyeCenterPoint_RenderThread(EStereoscopicEye::eSSE_LEFT_EYE), SrcRect, WindowRect);
 
@@ -398,7 +398,7 @@ FIntRect FDefaultSpectatorScreenController::Helpers::GetEyeCroppedToFitRect(FVec
 		const int32 HalfHeightDiff = FMath::TruncToInt(((float)SrcRect.Height() - DesiredSrcHeight) * 0.5f);
 		OutRect.Min.Y += HalfHeightDiff;
 		OutRect.Max.Y -= HalfHeightDiff;
-		const int32 DesiredCenterAdjustment = FMath::TruncToInt((EyeCenterPoint.Y - 0.5f) * (float)SrcRect.Height());
+		const int32 DesiredCenterAdjustment = FMath::TruncToInt(((float)EyeCenterPoint.Y - 0.5f) * (float)SrcRect.Height());
 		const int32 ActualCenterAdjustment = FMath::Clamp(DesiredCenterAdjustment, -HalfHeightDiff, HalfHeightDiff);
 		OutRect.Min.Y += ActualCenterAdjustment;
 		OutRect.Max.Y += ActualCenterAdjustment;
@@ -411,7 +411,7 @@ FIntRect FDefaultSpectatorScreenController::Helpers::GetEyeCroppedToFitRect(FVec
 		const int32 HalfWidthDiff = FMath::TruncToInt(((float)SrcRect.Width() - DesiredSrcWidth) * 0.5f);
 		OutRect.Min.X += HalfWidthDiff;
 		OutRect.Max.X -= HalfWidthDiff;
-		const int32 DesiredCenterAdjustment = FMath::TruncToInt((EyeCenterPoint.X - 0.5f) * (float)SrcRect.Width());
+		const int32 DesiredCenterAdjustment = FMath::TruncToInt(((float)EyeCenterPoint.X - 0.5f) * (float)SrcRect.Width());
 		const int32 ActualCenterAdjustment = FMath::Clamp(DesiredCenterAdjustment, -HalfWidthDiff, HalfWidthDiff);
 		OutRect.Min.X += ActualCenterAdjustment;
 		OutRect.Max.X += ActualCenterAdjustment;
diff --git a/Engine/Source/Runtime/HeadMountedDisplay/Private/DefaultStereoLayers.cpp b/Engine/Source/Runtime/HeadMountedDisplay/Private/DefaultStereoLayers.cpp
--- a/Engine/Source/Runtime/HeadMountedDisplay/Private/DefaultStereoLayers.cpp
+++ b/Engine/Source/Runtime/HeadMountedDisplay/Private/DefaultStereoLayers.cpp
@@ -130,8 +130,8 @@ void FDefaultStereoLayers::StereoLayerRender(FRHICommandListImmediate& RHICmdLis
 			const FRHITexture2D* Tex2D = Layer.Texture->GetTexture2D();
 			if (Tex2D)
 			{
-				const float SizeX = Tex2D->GetSizeX();
-				const float SizeY = Tex2D->GetSizeY();
+				const float SizeX = (float)Tex2D->GetSizeX();
+				const float SizeY = (float)Tex2D->GetSizeY();
 				if (SizeX != 0)
 				{
 					const float AspectRatio = SizeY / SizeX;
@@ -158,7 +158,7 @@ void FDefaultStereoLayers::StereoLayerRender(FRHICommandListImmediate& RHICmdLis
 		RendererModule.DrawRectangle(
 			RHICmdList,
 			0.0f, 0.0f,
-			TargetSize.X, TargetSize.Y,
+			(float)TargetSize.X, (float)TargetSize.Y,
 			0.0f, 0.0f,
 			1.0f, 1.0f,
 			TargetSize,
diff --git a/Engine/Source/Runtime/HeadMountedDisplay/Private/HeadMountedDisplayTypes.cpp b/Engine/Source/Runtime/HeadMountedDisplay/Private/HeadMountedDisplayTypes.cpp
--- a/Engine/Source/Runtime/HeadMountedDisplay/Private/HeadMountedDisplayTypes.cpp
+++ b/Engine/Source/Runtime/HeadMountedDisplay/Private/HeadMountedDisplayTypes.cpp
@@ -42,8 +42,8 @@ void FHMDViewMesh::BuildMesh(const FVector2D Positions[], uint32 VertexCount, EH
 		if (MeshType == MT_HiddenArea)
 		{
 			// Remap from to NDC space [0 1] -> [-1 1]
-			Vertex.Position.X = (Position.X * 2.0f) - 1.0f;
-			Vertex.Position.Y = (Position.Y * 2.0f) - 1.0f;
+			Vertex.Position.X = ((float)Position.X * 2.0f) - 1.0f;
+			Vertex.Position.Y = ((float)Position.Y * 2.0f) - 1.0f;
 			Vertex.Position.Z = 1.0f;
 			Vertex.Position.W = 1.0f;
 
@@ -54,16 +54,16 @@ void FHMDViewMesh::BuildMesh(const FVector2D Positions[], uint32 VertexCount, EH
 		else
 		{
 			// Remap the viewport origin from the bottom left to the top left
-			Vertex.Position.X = Position.X;
-			Vertex.Position.Y = 1.0f - Position.Y;
+			Vertex.Position.X = (float)Position.X;
+			Vertex.Position.Y = 1.0f - (float)Position.Y;
 			Vertex.Position.Z = 0.0f;
 			Vertex.Position.W = 1.0f;
 
-			Vertex.UV.X = Position.X;
-			Vertex.UV.Y = 1.0f - Position.Y;
+			Vertex.UV.X = (float)Position.X;
+			Vertex.UV.Y = 1.0f - (float)Position.Y;
 		}
 
-		pIndices[VertexIndex] = VertexIndex;
+		pIndices[VertexIndex] = static_cast<uint16>(VertexIndex);
 	}
 
 	RHIUnlockBuffer(VertexBufferRHI);
diff --git a/Engine/Source/Runtime/HeadMountedDisplay/Private/MotionDelayBuffer.cpp b/Engine/Source/Runtime/HeadMountedDisplay/Private/MotionDelayBuffer.cpp
--- a/Engine/Source/Runtime/HeadMountedDisplay/Private/MotionDelayBuffer.cpp
+++ b/Engine/Source/Runtime/HeadMountedDisplay/Private/MotionDelayBuffer.cpp
@@ -423,7 +423,7 @@ static void MotionDelayClient_Impl::CalulateDelayTransform(uint32 DesiredDelay,
 		FTransform TransformA = MakeTransform(SampleAIndex);
 		FTransform TransformB = MakeTransform(SampleBIndex);
 
-		const float BlendAlpha = (CurrentTime - SampleA.TimeStamp - DelaySeconds) / (SampleB.TimeStamp - SampleA.TimeStamp);
+		const float BlendAlpha = (float)((CurrentTime - SampleA.TimeStamp - DelaySeconds) / (SampleB.TimeStamp - SampleA.TimeStamp));
 		TransformOut.Blend(TransformA, TransformB, BlendAlpha);
 	}
 }
