Copy emulated face locked layers to spectator screen.

diff --git a/Engine/Plugins/Runtime/OpenXR/Source/OpenXRHMD/Private/OpenXRHMD.cpp b/Engine/Plugins/Runtime/OpenXR/Source/OpenXRHMD/Private/OpenXRHMD.cpp
--- a/Engine/Plugins/Runtime/OpenXR/Source/OpenXRHMD/Private/OpenXRHMD.cpp
+++ b/Engine/Plugins/Runtime/OpenXR/Source/OpenXRHMD/Private/OpenXRHMD.cpp
@@ -70,6 +70,12 @@ static TAutoConsoleVariable<bool> CVarOpenXRForceStereoLayerEmulation(
 	TEXT("Force the emulation of stereo layers instead of using native ones (if supported).\n"),
 	ECVF_Default);
 
+static TAutoConsoleVariable<bool> CVarOpenXRDoNotCopyEmulatedLayersToSpectatorScreen(
+	TEXT("xr.OpenXRDoNotCopyEmulatedLayersToSpectatorScreen"),
+	false,
+	TEXT("If face locked stereo layers emulation is active, avoid copying the face locked stereo layers to the spectator screen.\n"),
+	ECVF_Default);
+
 namespace {
 	static TSet<XrEnvironmentBlendMode> SupportedBlendModes{ XR_ENVIRONMENT_BLEND_MODE_ALPHA_BLEND, XR_ENVIRONMENT_BLEND_MODE_ADDITIVE, XR_ENVIRONMENT_BLEND_MODE_OPAQUE };
 	static TSet<XrViewConfigurationType> SupportedViewConfigurations{ XR_VIEW_CONFIGURATION_TYPE_PRIMARY_STEREO, XR_VIEW_CONFIGURATION_TYPE_PRIMARY_QUAD_VARJO };
@@ -3656,7 +3662,9 @@ void FOpenXRHMD::RenderTexture_RenderThread(class FRHICommandListImmediate& RHIC
 {
 	if (SpectatorScreenController)
 	{
-		SpectatorScreenController->RenderSpectatorScreen_RenderThread(RHICmdList, BackBuffer, SrcTexture, WindowSize);
+		const bool bShouldPassLayersTexture = PipelinedLayerStateRendering.EmulatedLayerState.bIsFaceLockedLayerEmulationActive && PipelinedLayerStateRendering.EmulatedLayerState.EmulationSwapchain && !CVarOpenXRDoNotCopyEmulatedLayersToSpectatorScreen.GetValueOnRenderThread();
+		const FTexture2DRHIRef LayersTexture = bShouldPassLayersTexture ? PipelinedLayerStateRendering.EmulatedLayerState.EmulationSwapchain->GetTextureRef() : nullptr;
+		SpectatorScreenController->RenderSpectatorScreen_RenderThread(RHICmdList, BackBuffer, SrcTexture, LayersTexture, WindowSize);
 	}
 }
 
diff --git a/Engine/Source/Runtime/HeadMountedDisplay/Private/DefaultSpectatorScreenController.cpp b/Engine/Source/Runtime/HeadMountedDisplay/Private/DefaultSpectatorScreenController.cpp
--- a/Engine/Source/Runtime/HeadMountedDisplay/Private/DefaultSpectatorScreenController.cpp
+++ b/Engine/Source/Runtime/HeadMountedDisplay/Private/DefaultSpectatorScreenController.cpp
@@ -234,11 +234,18 @@ void FDefaultSpectatorScreenController::UpdateSpectatorScreenMode_RenderThread()
 }
 
 void FDefaultSpectatorScreenController::RenderSpectatorScreen_RenderThread(FRHICommandListImmediate& RHICmdList, FRHITexture2D* BackBuffer, FTexture2DRHIRef SrcTexture, FVector2D WindowSize)
+{
+	RenderSpectatorScreen_RenderThread(RHICmdList, BackBuffer, SrcTexture, nullptr, WindowSize);
+}
+
+void FDefaultSpectatorScreenController::RenderSpectatorScreen_RenderThread(FRHICommandListImmediate& RHICmdList, FRHITexture2D* BackBuffer, FTexture2DRHIRef SrcTexture, FTexture2DRHIRef LayersTexture, FVector2D WindowSize)
 {
 	SCOPED_NAMED_EVENT_TEXT("RenderSocialScreen_RenderThread()", FColor::Magenta);
 
 	check(IsInRenderingThread());
 
+	StereoLayersTexture = LayersTexture;
+
 	if (SpectatorScreenDelegate_RenderThread.IsBound())
 	{
 		SCOPED_DRAW_EVENT(RHICmdList, SpectatorScreen)
@@ -247,9 +254,10 @@ void FDefaultSpectatorScreenController::RenderSpectatorScreen_RenderThread(FRHIC
 
 	// Apply the debug canvas layer.
 	IStereoLayers* StereoLayers = HMDDevice->GetStereoLayers();
-	if (StereoLayers)
+	if (StereoLayers && !LayersTexture)
 	{
 		const FIntRect DstRect(0, 0, BackBuffer->GetSizeX(), BackBuffer->GetSizeY());
+
 		for (int32 LayerID : DebugCanvasLayerIDs)
 		{
 			FTextureRHIRef LayerTexture = nullptr, HMDNull = nullptr;
@@ -267,18 +275,27 @@ void FDefaultSpectatorScreenController::RenderSpectatorScreen_RenderThread(FRHIC
 	}
 }
 
-
 FIntRect FDefaultSpectatorScreenController::GetFullFlatEyeRect_RenderThread(FTexture2DRHIRef EyeTexture)
 {
 	return HMDDevice->GetFullFlatEyeRect_RenderThread(EyeTexture);
 }
 
+void FDefaultSpectatorScreenController::CopyEmulatedLayers(FRHICommandListImmediate& RHICmdList, FTexture2DRHIRef TargetTexture, const FIntRect SrcRect, const FIntRect DstRect)
+{
+	if (StereoLayersTexture)
+	{
+		HMDDevice->CopyTexture_RenderThread(RHICmdList, StereoLayersTexture, SrcRect, TargetTexture, DstRect, false, false); 
+	}
+	StereoLayersTexture = nullptr;
+}
+
 void FDefaultSpectatorScreenController::RenderSpectatorModeUndistorted(FRHICommandListImmediate& RHICmdList, FTexture2DRHIRef TargetTexture, FTexture2DRHIRef EyeTexture, FTexture2DRHIRef OtherTexture, FVector2D WindowSize)
 {
 	const FIntRect SrcRect(0, 0, EyeTexture->GetSizeX(), EyeTexture->GetSizeY());
 	const FIntRect DstRect(0, 0, TargetTexture->GetSizeX(), TargetTexture->GetSizeY());
 
 	HMDDevice->CopyTexture_RenderThread(RHICmdList, EyeTexture, SrcRect, TargetTexture, DstRect, false, true);
+	CopyEmulatedLayers(RHICmdList, TargetTexture, SrcRect, DstRect);
 }
 
 void FDefaultSpectatorScreenController::RenderSpectatorModeDistorted(FRHICommandListImmediate& RHICmdList, FTexture2DRHIRef TargetTexture, FTexture2DRHIRef EyeTexture, FTexture2DRHIRef OtherTexture, FVector2D WindowSize)
@@ -294,6 +311,7 @@ void FDefaultSpectatorScreenController::RenderSpectatorModeSingleEye(FRHICommand
 	const FIntRect DstRect(0, 0, TargetTexture->GetSizeX(), TargetTexture->GetSizeY());
 
 	HMDDevice->CopyTexture_RenderThread(RHICmdList, EyeTexture, SrcRect, TargetTexture, DstRect, false, true);
+	CopyEmulatedLayers(RHICmdList, TargetTexture, SrcRect, DstRect);
 }
 
 void FDefaultSpectatorScreenController::RenderSpectatorModeSingleEyeLetterboxed(FRHICommandListImmediate& RHICmdList, FTexture2DRHIRef TargetTexture, FTexture2DRHIRef EyeTexture, FTexture2DRHIRef OtherTexture, FVector2D WindowSize)
@@ -303,6 +321,7 @@ void FDefaultSpectatorScreenController::RenderSpectatorModeSingleEyeLetterboxed(
 	const FIntRect DstRectLetterboxed = Helpers::GetLetterboxedDestRect(SrcRect, DstRect);
 
 	HMDDevice->CopyTexture_RenderThread(RHICmdList, EyeTexture, SrcRect, TargetTexture, DstRectLetterboxed, true, true);
+	CopyEmulatedLayers(RHICmdList, TargetTexture, SrcRect, DstRectLetterboxed);
 }
 
 void FDefaultSpectatorScreenController::RenderSpectatorModeSingleEyeCroppedToFill(FRHICommandListImmediate& RHICmdList, FTexture2DRHIRef TargetTexture, FTexture2DRHIRef EyeTexture, FTexture2DRHIRef OtherTexture, FVector2D WindowSize)
@@ -314,6 +333,7 @@ void FDefaultSpectatorScreenController::RenderSpectatorModeSingleEyeCroppedToFil
 	const FIntRect SrcCroppedToFitRect = Helpers::GetEyeCroppedToFitRect(HMDDevice->GetEyeCenterPoint_RenderThread(EStereoscopicEye::eSSE_LEFT_EYE), SrcRect, WindowRect);
 
 	HMDDevice->CopyTexture_RenderThread(RHICmdList, EyeTexture, SrcCroppedToFitRect, TargetTexture, DstRect, false, true);
+	CopyEmulatedLayers(RHICmdList, TargetTexture, SrcCroppedToFitRect, DstRect);
 }
 
 void FDefaultSpectatorScreenController::RenderSpectatorModeTexture(FRHICommandListImmediate& RHICmdList, FTexture2DRHIRef TargetTexture, FTexture2DRHIRef EyeTexture, FTexture2DRHIRef OtherTexture, FVector2D WindowSize)
@@ -350,12 +370,14 @@ void FDefaultSpectatorScreenController::RenderSpectatorModeMirrorAndTexture(FRHI
 	if (SpectatorScreenModeTexturePlusEyeLayout_RenderThread.bDrawEyeFirst)
 	{
 		HMDDevice->CopyTexture_RenderThread(RHICmdList, EyeTexture, CroppedEyeSrcRect, TargetTexture, EyeDstRect, bClearBlack, true);
+		CopyEmulatedLayers(RHICmdList, TargetTexture, CroppedEyeSrcRect, EyeDstRect);
 		HMDDevice->CopyTexture_RenderThread(RHICmdList, OtherTextureLocal, OtherSrcRect, TargetTexture, OtherDstRect, false, !SpectatorScreenModeTexturePlusEyeLayout_RenderThread.bUseAlpha);
 	}
 	else
 	{
 		HMDDevice->CopyTexture_RenderThread(RHICmdList, OtherTextureLocal, OtherSrcRect, TargetTexture, OtherDstRect, bClearBlack, true);
 		HMDDevice->CopyTexture_RenderThread(RHICmdList, EyeTexture, CroppedEyeSrcRect, TargetTexture, EyeDstRect, false, true);
+		CopyEmulatedLayers(RHICmdList, TargetTexture, CroppedEyeSrcRect, EyeDstRect);
 	}
 }
 
diff --git a/Engine/Source/Runtime/HeadMountedDisplay/Public/DefaultSpectatorScreenController.h b/Engine/Source/Runtime/HeadMountedDisplay/Public/DefaultSpectatorScreenController.h
--- a/Engine/Source/Runtime/HeadMountedDisplay/Public/DefaultSpectatorScreenController.h
+++ b/Engine/Source/Runtime/HeadMountedDisplay/Public/DefaultSpectatorScreenController.h
@@ -33,7 +33,8 @@ public:
 	// Implementation methods called by HMD
 	virtual void BeginRenderViewFamily();
 	virtual void UpdateSpectatorScreenMode_RenderThread();
-	virtual void RenderSpectatorScreen_RenderThread(FRHICommandListImmediate& RHICmdList, FRHITexture2D* BackBuffer, FTexture2DRHIRef RenderTarget, FVector2D WindowSize);
+	virtual void RenderSpectatorScreen_RenderThread(FRHICommandListImmediate& RHICmdList, FRHITexture2D* BackBuffer, FTexture2DRHIRef SrcTexture, FVector2D WindowSize);
+	virtual void RenderSpectatorScreen_RenderThread(FRHICommandListImmediate& RHICmdList, FRHITexture2D* BackBuffer, FTexture2DRHIRef SrcTexture, FTexture2DRHIRef LayersTexture, FVector2D WindowSize);
 
 protected:
 	friend struct FRHISetSpectatorScreenTexture;
@@ -74,5 +75,9 @@ protected:
 	};
 
 private:
+	void CopyEmulatedLayers(FRHICommandListImmediate& RHICmdList, FTexture2DRHIRef TargetTexture, const FIntRect SrcRect, const FIntRect DstRect);
+
 	FHeadMountedDisplayBase* HMDDevice;
+	// Face locked stereo layers are composited to a single texture which has to be copied over to the spectator screen.
+	FTexture2DRHIRef StereoLayersTexture;
 };
\ No newline at end of file
