Fixed editor crash happening when accessing vr.SpectatorScreenMode console variable without a controller.

diff --git a/Engine/Source/Runtime/HeadMountedDisplay/Private/HeadMountedDisplayConsoleCommands.cpp b/Engine/Source/Runtime/HeadMountedDisplay/Private/HeadMountedDisplayConsoleCommands.cpp
--- a/Engine/Source/Runtime/HeadMountedDisplay/Private/HeadMountedDisplayConsoleCommands.cpp
+++ b/Engine/Source/Runtime/HeadMountedDisplay/Private/HeadMountedDisplayConsoleCommands.cpp
@@ -245,7 +245,7 @@ static void SpectatorScreenMode(const TArray<FString>& Args, UWorld* , FOutputDe
 	ISpectatorScreenController* const Controller = HMDConsoleCommandsHelpers::GetSpectatorScreenController();
 	if (Controller == nullptr)
 	{
-		Ar.Logf(ELogVerbosity::Error, TEXT("SpectatorScreenMode is not controllable now, cannot change or get mode."), *Args[0]);
+		Ar.Logf(ELogVerbosity::Error, TEXT("SpectatorScreenMode is not controllable now, cannot change or get mode."));
 		return;
 	}
 
