Added tooltip to MotionControllerComponent MotionSource.

diff --git a/Engine/Source/Runtime/HeadMountedDisplay/Public/MotionControllerComponent.h b/Engine/Source/Runtime/HeadMountedDisplay/Public/MotionControllerComponent.h
--- a/Engine/Source/Runtime/HeadMountedDisplay/Public/MotionControllerComponent.h
+++ b/Engine/Source/Runtime/HeadMountedDisplay/Public/MotionControllerComponent.h
@@ -28,6 +28,7 @@ class HEADMOUNTEDDISPLAY_API UMotionControllerComponent : public UPrimitiveCompo
 	UPROPERTY(BlueprintSetter = SetTrackingSource, BlueprintGetter = GetTrackingSource, Category = "MotionController")
 	EControllerHand Hand_DEPRECATED;
 
+	/** Defines which pose this component should receive from the OpenXR Runtime. Left/Right MotionSource is the same as LeftGrip/RightGrip. See OpenXR specification for details on poses. */
 	UPROPERTY(EditAnywhere, BlueprintReadWrite, BlueprintSetter = SetTrackingMotionSource, Category = "MotionController")
 	FName MotionSource;
 
