Subscribed XRDeviceVisualizationComponent to OnXRInteractionProfileChanged delegate.

xrGetCurrentInteractionProfile returns 0 for the first few frames. If the MotionControllerComponent tries to start the device visualization before xrGetCurrentInteractionProfile has started returning a valid interaction profile the OpenXRAssetManager will return a null static mesh. Receving the callback from OnXRInteractionProfileChanged allows to set the correct static mesh when the correct interaction profile is set.

diff --git a/Engine/Source/Runtime/HeadMountedDisplay/Private/XRDeviceVisualizationComponent.cpp b/Engine/Source/Runtime/HeadMountedDisplay/Private/XRDeviceVisualizationComponent.cpp
--- a/Engine/Source/Runtime/HeadMountedDisplay/Private/XRDeviceVisualizationComponent.cpp
+++ b/Engine/Source/Runtime/HeadMountedDisplay/Private/XRDeviceVisualizationComponent.cpp
@@ -7,6 +7,7 @@
 #include "IXRTrackingSystem.h"
 #include "Misc/Attribute.h"
 #include "MotionControllerComponent.h"
+#include "XRTrackingSystemBase.h"
 
 #include UE_INLINE_GENERATED_CPP_BY_NAME(XRDeviceVisualizationComponent)
 
@@ -36,14 +37,28 @@ void UXRDeviceVisualizationComponent::BeginPlay() {
 		MotionSource = ParentController->GetTrackingMotionSource();
 		// Wait for controller to activate the rendering for this component.
 		bIsRenderingActive = false;
+		FXRTrackingSystemDelegates::OnXRInteractionProfileChanged.AddUObject(this, &UXRDeviceVisualizationComponent::OnInteractionProfileChanged);
 	}
 	else 
 	{
 		UE_LOG(LogXRDeviceVisualizationComponent, Error, TEXT("No parent UMotionControllerComponent for %s was found in the hierarchy. This component will not render anything."), *GetName());
 	}
+}
+
+//=============================================================================
+void UXRDeviceVisualizationComponent::OnInteractionProfileChanged()
+{
+	bInteractionProfilePresent = true;
 	RefreshMesh();
 }
 
+//=============================================================================
+void UXRDeviceVisualizationComponent::OnUnregister()
+{
+	Super::OnUnregister();
+	FXRTrackingSystemDelegates::OnXRInteractionProfileChanged.RemoveAll(this);
+}
+
 //=============================================================================
 UMotionControllerComponent* UXRDeviceVisualizationComponent::FindParentMotionController()
 {	
@@ -65,8 +80,8 @@ void UXRDeviceVisualizationComponent::SetIsVisualizationActive(bool bNewVisualiz
 	if (bNewVisualizationState != bIsVisualizationActive) 
 	{
 		bIsVisualizationActive = bNewVisualizationState;
-		RefreshMesh();
 	}
+	RefreshMesh();
 }
 
 //=============================================================================
@@ -165,7 +180,6 @@ void UXRDeviceVisualizationComponent::RefreshMesh()
 
 				const EObjectFlags SubObjFlags = RF_Transactional | RF_TextExportTransient;
 				UPrimitiveComponent* RenderComponent = AssetSys->CreateRenderComponent(DeviceId, GetOwner(), SubObjFlags, /*bForceSynchronous=*/false, LoadCompleteDelegate);
-				
 				if (RenderComponent != nullptr) 
 				{
 					StaticMeshComponent = Cast<UStaticMeshComponent>(RenderComponent);
@@ -189,7 +203,14 @@ void UXRDeviceVisualizationComponent::RefreshMesh()
 				{
 					DisplayModelLoadState = EModelLoadStatus::Unloaded;
 					SetStaticMesh(nullptr);
-					UE_LOG(LogXRDeviceVisualizationComponent, Error, TEXT("CreateRenderComponent returned a NULL UPrimitiveComponent*. Setting a NULL StaticMesh."));
+					// xrGetCurrentInteractionProfile returns 0 for the first bunch of frames.
+					// There is a race condition between the MotionController starting the rendering on this component and the interaction profile becoming available.
+					// If the MotionController starts the rendering before an interaction profile has been found, the RenderComponent will be null and we don't want
+					// to fire an error for this.
+					if (bInteractionProfilePresent)
+					{
+						UE_LOG(LogXRDeviceVisualizationComponent, Error, TEXT("CreateRenderComponent returned a NULL UPrimitiveComponent*. Setting a NULL StaticMesh."));
+					}
 				}
 			}
 		}
diff --git a/Engine/Source/Runtime/HeadMountedDisplay/Public/XRDeviceVisualizationComponent.h b/Engine/Source/Runtime/HeadMountedDisplay/Public/XRDeviceVisualizationComponent.h
--- a/Engine/Source/Runtime/HeadMountedDisplay/Public/XRDeviceVisualizationComponent.h
+++ b/Engine/Source/Runtime/HeadMountedDisplay/Public/XRDeviceVisualizationComponent.h
@@ -66,6 +66,7 @@ class HEADMOUNTEDDISPLAY_API UXRDeviceVisualizationComponent : public UStaticMes
 	FXRDeviceId DisplayDeviceId;
 
 	void BeginPlay() override;
+	virtual void OnUnregister() override;
 	void RefreshMesh();
 	UMotionControllerComponent* FindParentMotionController();
 	void SetMaterials(int32 MatCount);
@@ -81,4 +82,9 @@ public:
 
 	UFUNCTION(BlueprintSetter, Category="MotionController")
 	void SetIsRenderingActive(bool bRenderingIsActive);
+
+private:
+	void OnInteractionProfileChanged();
+
+	bool bInteractionProfilePresent = false;
 };
\ No newline at end of file
