StereoLayerManager potentially uses invalidated FLayerData
-Make a copy of the top FLayerData from which we can copy properties if we need to when adding new ones to the stack.

diff --git a/Engine/Source/Runtime/HeadMountedDisplay/Public/StereoLayerManager.h b/Engine/Source/Runtime/HeadMountedDisplay/Public/StereoLayerManager.h
--- a/Engine/Source/Runtime/HeadMountedDisplay/Public/StereoLayerManager.h
+++ b/Engine/Source/Runtime/HeadMountedDisplay/Public/StereoLayerManager.h
@@ -281,7 +281,7 @@ public:
 	virtual void PushLayerState(bool bPreserve = false) override
 	{
 		FScopeLock LockLayers(&LayerCritSect);
-		const auto& CurrentState = LayerState();
+		const FLayerData CurrentState = LayerState(); // Copy because we modify the Stack that LayerState returns an element of.
 
 		if (bPreserve)
 		{
