UE-181471 Delete SteamVR plugin.
-Deprecated in 5.1.  OpenXR should be used instead.
-A few instances of the string "SteamVR" remain.  In some case we are comparing to the OpenXR Runtime name, which remains correct.  In some we are searching input profile paths for SteamVR, which also remains correct.  VRScouting is doing the input profile path and then overriding VREditor's view of the xr plugin name, so at first glance VREditor still seems to be using SteamVR, but it actually does not depend on that (That code will likely be further refactored as time goes on).

diff --git a/Engine/Plugins/Runtime/Steam/SteamVR/Source/SteamVR/Classes/SteamVRChaperoneComponent.h b/Engine/Plugins/Runtime/Steam/SteamVR/Source/SteamVR/Classes/SteamVRChaperoneComponent.h
deleted file mode 100644
diff --git a/Engine/Plugins/Runtime/Steam/SteamVR/Source/SteamVR/Classes/SteamVRFunctionLibrary.h b/Engine/Plugins/Runtime/Steam/SteamVR/Source/SteamVR/Classes/SteamVRFunctionLibrary.h
deleted file mode 100644
diff --git a/Engine/Plugins/Runtime/Steam/SteamVR/Source/SteamVR/Classes/SteamVRHQLayerShape.h b/Engine/Plugins/Runtime/Steam/SteamVR/Source/SteamVR/Classes/SteamVRHQLayerShape.h
deleted file mode 100644
diff --git a/Engine/Plugins/Runtime/Steam/SteamVR/Source/SteamVR/Private/SteamVRAssetManager.cpp b/Engine/Plugins/Runtime/Steam/SteamVR/Source/SteamVR/Private/SteamVRAssetManager.cpp
deleted file mode 100644
diff --git a/Engine/Plugins/Runtime/Steam/SteamVR/Source/SteamVR/Private/SteamVRAssetManager.h b/Engine/Plugins/Runtime/Steam/SteamVR/Source/SteamVR/Private/SteamVRAssetManager.h
deleted file mode 100644
diff --git a/Engine/Plugins/Runtime/Steam/SteamVR/Source/SteamVR/Private/SteamVRChaperoneComponent.cpp b/Engine/Plugins/Runtime/Steam/SteamVR/Source/SteamVR/Private/SteamVRChaperoneComponent.cpp
deleted file mode 100644
diff --git a/Engine/Plugins/Runtime/Steam/SteamVR/Source/SteamVR/Private/SteamVRFunctionLibrary.cpp b/Engine/Plugins/Runtime/Steam/SteamVR/Source/SteamVR/Private/SteamVRFunctionLibrary.cpp
deleted file mode 100644
diff --git a/Engine/Plugins/Runtime/Steam/SteamVR/Source/SteamVR/Private/SteamVRHMD.cpp b/Engine/Plugins/Runtime/Steam/SteamVR/Source/SteamVR/Private/SteamVRHMD.cpp
deleted file mode 100644
diff --git a/Engine/Plugins/Runtime/Steam/SteamVR/Source/SteamVR/Private/SteamVRHMD.h b/Engine/Plugins/Runtime/Steam/SteamVR/Source/SteamVR/Private/SteamVRHMD.h
deleted file mode 100644
diff --git a/Engine/Plugins/Runtime/Steam/SteamVR/Source/SteamVR/Private/SteamVRHQLayerShape.cpp b/Engine/Plugins/Runtime/Steam/SteamVR/Source/SteamVR/Private/SteamVRHQLayerShape.cpp
deleted file mode 100644
diff --git a/Engine/Plugins/Runtime/Steam/SteamVR/Source/SteamVR/Private/SteamVRPrivate.h b/Engine/Plugins/Runtime/Steam/SteamVR/Source/SteamVR/Private/SteamVRPrivate.h
deleted file mode 100644
diff --git a/Engine/Plugins/Runtime/Steam/SteamVR/Source/SteamVR/Private/SteamVRRender.cpp b/Engine/Plugins/Runtime/Steam/SteamVR/Source/SteamVR/Private/SteamVRRender.cpp
deleted file mode 100644
diff --git a/Engine/Plugins/Runtime/Steam/SteamVR/Source/SteamVR/Private/SteamVRSpectatorScreen.cpp b/Engine/Plugins/Runtime/Steam/SteamVR/Source/SteamVR/Private/SteamVRSpectatorScreen.cpp
deleted file mode 100644
diff --git a/Engine/Plugins/Runtime/Steam/SteamVR/Source/SteamVR/Private/SteamVRSplash.cpp b/Engine/Plugins/Runtime/Steam/SteamVR/Source/SteamVR/Private/SteamVRSplash.cpp
deleted file mode 100644
diff --git a/Engine/Plugins/Runtime/Steam/SteamVR/Source/SteamVR/Private/SteamVRSplash.h b/Engine/Plugins/Runtime/Steam/SteamVR/Source/SteamVR/Private/SteamVRSplash.h
deleted file mode 100644
diff --git a/Engine/Plugins/Runtime/Steam/SteamVR/Source/SteamVR/Private/SteamVRStereoLayers.cpp b/Engine/Plugins/Runtime/Steam/SteamVR/Source/SteamVR/Private/SteamVRStereoLayers.cpp
deleted file mode 100644
diff --git a/Engine/Plugins/Runtime/Steam/SteamVR/Source/SteamVR/Public/ISteamVRPlugin.h b/Engine/Plugins/Runtime/Steam/SteamVR/Source/SteamVR/Public/ISteamVRPlugin.h
deleted file mode 100644
diff --git a/Engine/Plugins/Runtime/Steam/SteamVR/Source/SteamVR/Public/SteamVRStereoLayers.h b/Engine/Plugins/Runtime/Steam/SteamVR/Source/SteamVR/Public/SteamVRStereoLayers.h
deleted file mode 100644
diff --git a/Engine/Plugins/Runtime/Steam/SteamVR/Source/SteamVR/SteamVR.Build.cs b/Engine/Plugins/Runtime/Steam/SteamVR/Source/SteamVR/SteamVR.Build.cs
deleted file mode 100644
diff --git a/Engine/Plugins/Runtime/Steam/SteamVR/Source/SteamVREditor/Private/AnimGraphNode_SteamVRInputAnimPose.cpp b/Engine/Plugins/Runtime/Steam/SteamVR/Source/SteamVREditor/Private/AnimGraphNode_SteamVRInputAnimPose.cpp
deleted file mode 100644
diff --git a/Engine/Plugins/Runtime/Steam/SteamVR/Source/SteamVREditor/Private/AnimGraphNode_SteamVRSetWristTransform.cpp b/Engine/Plugins/Runtime/Steam/SteamVR/Source/SteamVREditor/Private/AnimGraphNode_SteamVRSetWristTransform.cpp
deleted file mode 100644
diff --git a/Engine/Plugins/Runtime/Steam/SteamVR/Source/SteamVREditor/Private/SteamVREditor.cpp b/Engine/Plugins/Runtime/Steam/SteamVR/Source/SteamVREditor/Private/SteamVREditor.cpp
deleted file mode 100644
diff --git a/Engine/Plugins/Runtime/Steam/SteamVR/Source/SteamVREditor/Private/SteamVREditorCommands.cpp b/Engine/Plugins/Runtime/Steam/SteamVR/Source/SteamVREditor/Private/SteamVREditorCommands.cpp
deleted file mode 100644
diff --git a/Engine/Plugins/Runtime/Steam/SteamVR/Source/SteamVREditor/Private/SteamVREditorCommands.h b/Engine/Plugins/Runtime/Steam/SteamVR/Source/SteamVREditor/Private/SteamVREditorCommands.h
deleted file mode 100644
diff --git a/Engine/Plugins/Runtime/Steam/SteamVR/Source/SteamVREditor/Private/SteamVREditorSettings.cpp b/Engine/Plugins/Runtime/Steam/SteamVR/Source/SteamVREditor/Private/SteamVREditorSettings.cpp
deleted file mode 100644
diff --git a/Engine/Plugins/Runtime/Steam/SteamVR/Source/SteamVREditor/Private/SteamVREditorStyle.cpp b/Engine/Plugins/Runtime/Steam/SteamVR/Source/SteamVREditor/Private/SteamVREditorStyle.cpp
deleted file mode 100644
diff --git a/Engine/Plugins/Runtime/Steam/SteamVR/Source/SteamVREditor/Private/SteamVREditorStyle.h b/Engine/Plugins/Runtime/Steam/SteamVR/Source/SteamVREditor/Private/SteamVREditorStyle.h
deleted file mode 100644
diff --git a/Engine/Plugins/Runtime/Steam/SteamVR/Source/SteamVREditor/Public/AnimGraphNode_SteamVRInputAnimPose.h b/Engine/Plugins/Runtime/Steam/SteamVR/Source/SteamVREditor/Public/AnimGraphNode_SteamVRInputAnimPose.h
deleted file mode 100644
diff --git a/Engine/Plugins/Runtime/Steam/SteamVR/Source/SteamVREditor/Public/AnimGraphNode_SteamVRSetWristTransform.h b/Engine/Plugins/Runtime/Steam/SteamVR/Source/SteamVREditor/Public/AnimGraphNode_SteamVRSetWristTransform.h
deleted file mode 100644
diff --git a/Engine/Plugins/Runtime/Steam/SteamVR/Source/SteamVREditor/Public/SteamVREditor.h b/Engine/Plugins/Runtime/Steam/SteamVR/Source/SteamVREditor/Public/SteamVREditor.h
deleted file mode 100644
diff --git a/Engine/Plugins/Runtime/Steam/SteamVR/Source/SteamVREditor/Public/SteamVREditorSettings.h b/Engine/Plugins/Runtime/Steam/SteamVR/Source/SteamVREditor/Public/SteamVREditorSettings.h
deleted file mode 100644
diff --git a/Engine/Plugins/Runtime/Steam/SteamVR/Source/SteamVREditor/SteamVREditor.Build.cs b/Engine/Plugins/Runtime/Steam/SteamVR/Source/SteamVREditor/SteamVREditor.Build.cs
deleted file mode 100644
diff --git a/Engine/Plugins/Runtime/Steam/SteamVR/Source/SteamVRInput/Private/SteamVRInput.cpp b/Engine/Plugins/Runtime/Steam/SteamVR/Source/SteamVRInput/Private/SteamVRInput.cpp
deleted file mode 100644
diff --git a/Engine/Plugins/Runtime/Steam/SteamVR/Source/SteamVRInput/Public/SteamVRInput.h b/Engine/Plugins/Runtime/Steam/SteamVR/Source/SteamVRInput/Public/SteamVRInput.h
deleted file mode 100644
diff --git a/Engine/Plugins/Runtime/Steam/SteamVR/Source/SteamVRInput/SteamVRInput.Build.cs b/Engine/Plugins/Runtime/Steam/SteamVR/Source/SteamVRInput/SteamVRInput.Build.cs
deleted file mode 100644
diff --git a/Engine/Plugins/Runtime/Steam/SteamVR/Source/SteamVRInputDevice/Private/AnimNode_SteamVRInputAnimPose.cpp b/Engine/Plugins/Runtime/Steam/SteamVR/Source/SteamVRInputDevice/Private/AnimNode_SteamVRInputAnimPose.cpp
deleted file mode 100644
diff --git a/Engine/Plugins/Runtime/Steam/SteamVR/Source/SteamVRInputDevice/Private/AnimNode_SteamVRSetWristTransform.cpp b/Engine/Plugins/Runtime/Steam/SteamVR/Source/SteamVRInputDevice/Private/AnimNode_SteamVRSetWristTransform.cpp
deleted file mode 100644
diff --git a/Engine/Plugins/Runtime/Steam/SteamVR/Source/SteamVRInputDevice/Private/SteamVRControllerKeys.h b/Engine/Plugins/Runtime/Steam/SteamVR/Source/SteamVRInputDevice/Private/SteamVRControllerKeys.h
deleted file mode 100644
diff --git a/Engine/Plugins/Runtime/Steam/SteamVR/Source/SteamVRInputDevice/Private/SteamVRInputDevice.cpp b/Engine/Plugins/Runtime/Steam/SteamVR/Source/SteamVRInputDevice/Private/SteamVRInputDevice.cpp
deleted file mode 100644
diff --git a/Engine/Plugins/Runtime/Steam/SteamVR/Source/SteamVRInputDevice/Private/SteamVRInputDevice.h b/Engine/Plugins/Runtime/Steam/SteamVR/Source/SteamVRInputDevice/Private/SteamVRInputDevice.h
deleted file mode 100644
diff --git a/Engine/Plugins/Runtime/Steam/SteamVR/Source/SteamVRInputDevice/Private/SteamVRInputDeviceFunctionLibrary.cpp b/Engine/Plugins/Runtime/Steam/SteamVR/Source/SteamVRInputDevice/Private/SteamVRInputDeviceFunctionLibrary.cpp
deleted file mode 100644
diff --git a/Engine/Plugins/Runtime/Steam/SteamVR/Source/SteamVRInputDevice/Private/SteamVRInputDeviceModule.cpp b/Engine/Plugins/Runtime/Steam/SteamVR/Source/SteamVRInputDevice/Private/SteamVRInputDeviceModule.cpp
deleted file mode 100644
diff --git a/Engine/Plugins/Runtime/Steam/SteamVR/Source/SteamVRInputDevice/Private/SteamVRInputTypes.h b/Engine/Plugins/Runtime/Steam/SteamVR/Source/SteamVRInputDevice/Private/SteamVRInputTypes.h
deleted file mode 100644
diff --git a/Engine/Plugins/Runtime/Steam/SteamVR/Source/SteamVRInputDevice/Private/SteamVRKnucklesKeys.h b/Engine/Plugins/Runtime/Steam/SteamVR/Source/SteamVRInputDevice/Private/SteamVRKnucklesKeys.h
deleted file mode 100644
diff --git a/Engine/Plugins/Runtime/Steam/SteamVR/Source/SteamVRInputDevice/Private/SteamVRSkeletonDefinition.cpp b/Engine/Plugins/Runtime/Steam/SteamVR/Source/SteamVRInputDevice/Private/SteamVRSkeletonDefinition.cpp
deleted file mode 100644
diff --git a/Engine/Plugins/Runtime/Steam/SteamVR/Source/SteamVRInputDevice/Private/SteamVRTrackingRefComponent.cpp b/Engine/Plugins/Runtime/Steam/SteamVR/Source/SteamVRInputDevice/Private/SteamVRTrackingRefComponent.cpp
deleted file mode 100644
diff --git a/Engine/Plugins/Runtime/Steam/SteamVR/Source/SteamVRInputDevice/Private/UE4HandSkeletonDefinition.cpp b/Engine/Plugins/Runtime/Steam/SteamVR/Source/SteamVRInputDevice/Private/UE4HandSkeletonDefinition.cpp
deleted file mode 100644
diff --git a/Engine/Plugins/Runtime/Steam/SteamVR/Source/SteamVRInputDevice/Private/UE4HandSkeletonDefinition.h b/Engine/Plugins/Runtime/Steam/SteamVR/Source/SteamVRInputDevice/Private/UE4HandSkeletonDefinition.h
deleted file mode 100644
diff --git a/Engine/Plugins/Runtime/Steam/SteamVR/Source/SteamVRInputDevice/Public/AnimNode_SteamVRInputAnimPose.h b/Engine/Plugins/Runtime/Steam/SteamVR/Source/SteamVRInputDevice/Public/AnimNode_SteamVRInputAnimPose.h
deleted file mode 100644
diff --git a/Engine/Plugins/Runtime/Steam/SteamVR/Source/SteamVRInputDevice/Public/AnimNode_SteamVRSetWristTransform.h b/Engine/Plugins/Runtime/Steam/SteamVR/Source/SteamVRInputDevice/Public/AnimNode_SteamVRSetWristTransform.h
deleted file mode 100644
diff --git a/Engine/Plugins/Runtime/Steam/SteamVR/Source/SteamVRInputDevice/Public/ISteamVRInputDeviceModule.h b/Engine/Plugins/Runtime/Steam/SteamVR/Source/SteamVRInputDevice/Public/ISteamVRInputDeviceModule.h
deleted file mode 100644
diff --git a/Engine/Plugins/Runtime/Steam/SteamVR/Source/SteamVRInputDevice/Public/SteamVRInputDeviceFunctionLibrary.h b/Engine/Plugins/Runtime/Steam/SteamVR/Source/SteamVRInputDevice/Public/SteamVRInputDeviceFunctionLibrary.h
deleted file mode 100644
diff --git a/Engine/Plugins/Runtime/Steam/SteamVR/Source/SteamVRInputDevice/Public/SteamVRInputPublic.h b/Engine/Plugins/Runtime/Steam/SteamVR/Source/SteamVRInputDevice/Public/SteamVRInputPublic.h
deleted file mode 100644
diff --git a/Engine/Plugins/Runtime/Steam/SteamVR/Source/SteamVRInputDevice/Public/SteamVRSkeletonDefinition.h b/Engine/Plugins/Runtime/Steam/SteamVR/Source/SteamVRInputDevice/Public/SteamVRSkeletonDefinition.h
deleted file mode 100644
diff --git a/Engine/Plugins/Runtime/Steam/SteamVR/Source/SteamVRInputDevice/Public/SteamVRTrackingRefComponent.h b/Engine/Plugins/Runtime/Steam/SteamVR/Source/SteamVRInputDevice/Public/SteamVRTrackingRefComponent.h
deleted file mode 100644
diff --git a/Engine/Plugins/Runtime/Steam/SteamVR/Source/SteamVRInputDevice/SteamVRInputDevice.Build.cs b/Engine/Plugins/Runtime/Steam/SteamVR/Source/SteamVRInputDevice/SteamVRInputDevice.Build.cs
deleted file mode 100644
diff --git a/Engine/Plugins/Runtime/Steam/SteamVR/SteamVR.uplugin b/Engine/Plugins/Runtime/Steam/SteamVR/SteamVR.uplugin
deleted file mode 100644
diff --git a/Engine/Plugins/VirtualProduction/LiveLinkXR/Source/LiveLinkXR/Public/LiveLinkXRConnectionSettings.h b/Engine/Plugins/VirtualProduction/LiveLinkXR/Source/LiveLinkXR/Public/LiveLinkXRConnectionSettings.h
--- a/Engine/Plugins/VirtualProduction/LiveLinkXR/Source/LiveLinkXR/Public/LiveLinkXRConnectionSettings.h
+++ b/Engine/Plugins/VirtualProduction/LiveLinkXR/Source/LiveLinkXR/Public/LiveLinkXRConnectionSettings.h
@@ -11,7 +11,7 @@ struct LIVELINKXR_API FLiveLinkXRConnectionSettings
 	GENERATED_BODY()
 
 public:
-	/** Track all SteamVR tracker pucks */
+	/** Track all tracker pucks */
 	UPROPERTY(EditAnywhere, Category = "Connection Settings")
 	bool bTrackTrackers = true;
 
diff --git a/Engine/Source/Editor/GameProjectGeneration/Private/GameProjectUtils.cpp b/Engine/Source/Editor/GameProjectGeneration/Private/GameProjectUtils.cpp
--- a/Engine/Source/Editor/GameProjectGeneration/Private/GameProjectUtils.cpp
+++ b/Engine/Source/Editor/GameProjectGeneration/Private/GameProjectUtils.cpp
@@ -100,8 +100,7 @@ namespace
 	void SetXRPluginStates(const FProjectInformation& InProjectInfo, FProjectDescriptor& OutProject)
 	{
 		static const FString XRPlugins[] = {
-			TEXT("OculusVR"),
-			TEXT("SteamVR") };
+			TEXT("OculusVR")};
 
 		if (InProjectInfo.bEnableXR.IsSet() && 
 			InProjectInfo.bEnableXR.GetValue() == false)
diff --git a/Engine/Source/Runtime/ApplicationCore/Private/Linux/LinuxApplication.cpp b/Engine/Source/Runtime/ApplicationCore/Private/Linux/LinuxApplication.cpp
--- a/Engine/Source/Runtime/ApplicationCore/Private/Linux/LinuxApplication.cpp
+++ b/Engine/Source/Runtime/ApplicationCore/Private/Linux/LinuxApplication.cpp
@@ -2025,8 +2025,6 @@ void FLinuxApplication::SetForceFeedbackChannelValues(int32 ControllerId, const
 	// send vibration to externally-implemented devices
 	for (auto DeviceIt = ExternalInputDevices.CreateIterator(); DeviceIt; ++DeviceIt)
 	{
-		// *N.B 06/20/2016*: Ideally, we would want to use GetHapticDevice instead
-		// but they're not implemented for SteamController and SteamVRController
 		if ((*DeviceIt)->IsGamepadAttached())
 		{
 			(*DeviceIt)->SetChannelValues(ControllerId, Values);
diff --git a/Engine/Source/Runtime/ApplicationCore/Private/Windows/WindowsApplication.cpp b/Engine/Source/Runtime/ApplicationCore/Private/Windows/WindowsApplication.cpp
--- a/Engine/Source/Runtime/ApplicationCore/Private/Windows/WindowsApplication.cpp
+++ b/Engine/Source/Runtime/ApplicationCore/Private/Windows/WindowsApplication.cpp
@@ -2835,8 +2835,6 @@ void FWindowsApplication::SetForceFeedbackChannelValues(int32 ControllerId, cons
 	// send vibration to externally-implemented devices
 	for( auto DeviceIt = ExternalInputDevices.CreateIterator(); DeviceIt; ++DeviceIt )
 	{
-		// *N.B 06/20/2016*: Ideally, we would want to use GetHapticDevice instead
-		// but they're not implemented for SteamController and SteamVRController
 		if ((*DeviceIt)->IsGamepadAttached()) 
 		{
 			(*DeviceIt)->SetChannelValues(ControllerId, Values);
diff --git a/Engine/Source/Runtime/HeadMountedDisplay/Public/HeadMountedDisplayFunctionLibrary.h b/Engine/Source/Runtime/HeadMountedDisplay/Public/HeadMountedDisplayFunctionLibrary.h
--- a/Engine/Source/Runtime/HeadMountedDisplay/Public/HeadMountedDisplayFunctionLibrary.h
+++ b/Engine/Source/Runtime/HeadMountedDisplay/Public/HeadMountedDisplayFunctionLibrary.h
@@ -276,7 +276,7 @@ class HEADMOUNTEDDISPLAY_API UHeadMountedDisplayFunctionLibrary : public UBluepr
 	/**
 	 * Cross XR-System query that will list all XR devices currently being tracked.
 	 *
-	 * @param  SystemId		(Optional) Specifies an explicit system to poll devices from (use if you want only devices belonging to one explicit XR ecosystem, e.g. 'OculusHMD', or 'SteamVR')
+	 * @param  SystemId		(Optional) Specifies an explicit system to poll devices from (use if you want only devices belonging to one explicit XR ecosystem, e.g. 'OculusHMD', or 'OpenXR')
 	 * @param  DeviceType	Specifies the type of device to query for - defaults to 'Any' (meaning 'All').
 	 *
 	 * @return A list of device identifiers matching the query. Use these to query and operate on the device (e.g. through GetDevicePose, AddDeviceVisualizationComponent, etc.)
diff --git a/Engine/Source/Runtime/HeadMountedDisplay/Public/XRAssetFunctionLibrary.h b/Engine/Source/Runtime/HeadMountedDisplay/Public/XRAssetFunctionLibrary.h
--- a/Engine/Source/Runtime/HeadMountedDisplay/Public/XRAssetFunctionLibrary.h
+++ b/Engine/Source/Runtime/HeadMountedDisplay/Public/XRAssetFunctionLibrary.h
@@ -47,7 +47,7 @@ public:
 	 *       will fail and return an invalid (null) object.
 	 *
 	 * @param  Target				The intended owner for the component to attach to.
-	 * @param  SystemName			(optional) Targets a specific XR system (i.e. 'Oculus', 'SteamVR', etc.). If left as 'None', then the first system found that can render the device will be used.
+	 * @param  SystemName			(optional) Targets a specific XR system (i.e. 'Oculus', 'OpenXR', etc.). If left as 'None', then the first system found that can render the device will be used.
 	 * @param  DeviceName			Source name of the specific device - expect the same names that the MotionControllerComponent's "MotionSource" field uses ('Left', 'Right', etc.).
 	 * @param  bManualAttachment	If set, will leave the component unattached (mirror's the same option on the generic AddComponent node). When unset the component will attach to the actor's root.
 	 * @param  RelativeTransform	Specifies the component initial transform (relative to its attach parent).
