UE-180807 XR EControllerHand use reduction.
-Refactored FXRMotionControllerBase EControllerHand GetControllerOrientationAndPosition and GetControllerTrackingStatus into a new FXRMotionControllerBaseLegacy that only legacy maintenence mode VR plugin motion controllers will derive from.  This allows us to make the new preferred api pure virtual and to eliminate implementations of the old deprecated api from plugins still in active development.
-Updated all VR plugins to comply.
-Removed "AnyHand" support from OpenXRHandTracking.  We don't support it in the controllers for OpenXR so we shouldn't support it in openxr hand tracking.

diff --git a/Engine/Plugins/Media/PixelStreaming/Source/PixelStreamingInput/Private/PixelStreamingInputHandler.cpp b/Engine/Plugins/Media/PixelStreaming/Source/PixelStreamingInput/Private/PixelStreamingInputHandler.cpp
--- a/Engine/Plugins/Media/PixelStreaming/Source/PixelStreamingInput/Private/PixelStreamingInputHandler.cpp
+++ b/Engine/Plugins/Media/PixelStreaming/Source/PixelStreamingInput/Private/PixelStreamingInputHandler.cpp
@@ -206,28 +206,23 @@ namespace UE::PixelStreamingInput
 		EControllerHand DeviceHand;
 		if (GetHandEnumForSourceName(MotionSource, DeviceHand))
 		{
-			return GetControllerOrientationAndPosition(ControllerIndex, DeviceHand, OutOrientation, OutPosition, WorldToMetersScale);
+			FPixelStreamingXRController Controller = XRControllers.FindRef(DeviceHand);
+			OutOrientation = Controller.Transform.Rotator();
+			OutPosition = Controller.Transform.GetTranslation();
+			return true;
 		}
 		return false;
 	}
 
-	bool FPixelStreamingInputHandler::GetControllerOrientationAndPosition(const int32 ControllerIndex, const EControllerHand DeviceHand, FRotator& OutOrientation, FVector& OutPosition, float WorldToMetersScale) const
+	ETrackingStatus FPixelStreamingInputHandler::GetControllerTrackingStatus(const int32 ControllerIndex, const FName MotionSource) const
 	{
-		if (FPixelStreamingHMD* HMD = IPixelStreamingHMDModule::Get().GetPixelStreamingHMD(); (HMD == nullptr || ControllerIndex == INDEX_NONE))
+		EControllerHand DeviceHand;
+		if (GetHandEnumForSourceName(MotionSource, DeviceHand))
 		{
-			return false;
+			const FPixelStreamingXRController* Controller = XRControllers.Find(DeviceHand);
+			return (Controller != nullptr) ? ETrackingStatus::Tracked : ETrackingStatus::NotTracked;
 		}
-
-		FPixelStreamingXRController Controller = XRControllers.FindRef(DeviceHand);
-		OutOrientation = Controller.Transform.Rotator();
-		OutPosition = Controller.Transform.GetTranslation();
-		return true;
-	}
-
-	ETrackingStatus FPixelStreamingInputHandler::GetControllerTrackingStatus(const int32 ControllerIndex, const EControllerHand DeviceHand) const
-	{
-		const FPixelStreamingXRController* Controller = XRControllers.Find(DeviceHand);
-		return (Controller != nullptr) ? ETrackingStatus::Tracked : ETrackingStatus::NotTracked;
+		return ETrackingStatus::NotTracked;
 	}
 
 	void FPixelStreamingInputHandler::EnumerateSources(TArray<FMotionControllerSource>& SourcesOut) const
diff --git a/Engine/Plugins/Media/PixelStreaming/Source/PixelStreamingInput/Private/PixelStreamingInputHandler.h b/Engine/Plugins/Media/PixelStreaming/Source/PixelStreamingInput/Private/PixelStreamingInputHandler.h
--- a/Engine/Plugins/Media/PixelStreaming/Source/PixelStreamingInput/Private/PixelStreamingInputHandler.h
+++ b/Engine/Plugins/Media/PixelStreaming/Source/PixelStreamingInput/Private/PixelStreamingInputHandler.h
@@ -57,8 +57,7 @@ namespace UE::PixelStreamingInput
 		// IMotionController Interface
 		virtual FName GetMotionControllerDeviceTypeName() const override;
 		virtual bool GetControllerOrientationAndPosition(const int32 ControllerIndex, const FName MotionSource, FRotator& OutOrientation, FVector& OutPosition, float WorldToMetersScale) const override;
-		virtual bool GetControllerOrientationAndPosition(const int32 ControllerIndex, const EControllerHand DeviceHand, FRotator& OutOrientation, FVector& OutPosition, float WorldToMetersScale) const override;
-		virtual ETrackingStatus GetControllerTrackingStatus(const int32 ControllerIndex, const EControllerHand DeviceHand) const override;
+		virtual ETrackingStatus GetControllerTrackingStatus(const int32 ControllerIndex, const FName MotionSource) const override;
 		virtual void EnumerateSources(TArray<FMotionControllerSource>& SourcesOut) const override;
 		// End IMotionController Interface
 
diff --git a/Engine/Plugins/Runtime/Oculus/OculusVR/Source/OculusInput/Private/OculusInput.h b/Engine/Plugins/Runtime/Oculus/OculusVR/Source/OculusInput/Private/OculusInput.h
--- a/Engine/Plugins/Runtime/Oculus/OculusVR/Source/OculusInput/Private/OculusInput.h
+++ b/Engine/Plugins/Runtime/Oculus/OculusVR/Source/OculusInput/Private/OculusInput.h
@@ -30,7 +30,7 @@ namespace OculusInput
 // FOculusInput
 //-------------------------------------------------------------------------------------------------
 
-class FOculusInput : public IInputDevice, public FXRMotionControllerBase, public IHapticDevice
+class FOculusInput : public IInputDevice, public FXRMotionControllerBaseLegacy, public IHapticDevice
 {
 	friend class FOculusHandTracking;
 
diff --git a/Engine/Plugins/Runtime/OpenXR/Source/OpenXRInput/Private/OpenXRInput.h b/Engine/Plugins/Runtime/OpenXR/Source/OpenXRInput/Private/OpenXRInput.h
--- a/Engine/Plugins/Runtime/OpenXR/Source/OpenXRInput/Private/OpenXRInput.h
+++ b/Engine/Plugins/Runtime/OpenXR/Source/OpenXRInput/Private/OpenXRInput.h
@@ -126,8 +126,6 @@ public:
 		virtual bool GetControllerOrientationAndPosition(const int32 ControllerIndex, const FName MotionSource, FRotator& OutOrientation, FVector& OutPosition, float WorldToMetersScale) const override;
 		virtual bool GetControllerOrientationAndPositionForTime(const int32 ControllerIndex, const FName MotionSource, FTimespan Time, bool& OutTimeWasUsed, FRotator& OutOrientation, FVector& OutPosition, bool& OutbProvidedLinearVelocity, FVector& OutLinearVelocity, bool& OutbProvidedAngularVelocity, FVector& OutAngularVelocityRadPerSec, bool& OutbProvidedLinearAcceleration, FVector& OutLinearAcceleration, float WorldToMetersScale) const override;
 		virtual ETrackingStatus GetControllerTrackingStatus(const int32 ControllerIndex, const FName MotionSource) const override;
-		virtual bool GetControllerOrientationAndPosition(const int32 ControllerIndex, const EControllerHand DeviceHand, FRotator& OutOrientation, FVector& OutPosition, float WorldToMetersScale) const override { check(false); return false; }
-		virtual ETrackingStatus GetControllerTrackingStatus(const int32 ControllerIndex, const EControllerHand DeviceHand) const override { check(false); return ETrackingStatus::NotTracked; }
 		virtual void EnumerateSources(TArray<FMotionControllerSource>& SourcesOut) const override;
 		virtual bool SetPlayerMappableInputConfig(TObjectPtr<class UPlayerMappableInputConfig> InputConfig) override;
 
diff --git a/Engine/Plugins/Runtime/OpenXRHandTracking/Source/OpenXRHandTracking/Private/OpenXRHandTracking.cpp b/Engine/Plugins/Runtime/OpenXRHandTracking/Source/OpenXRHandTracking/Private/OpenXRHandTracking.cpp
--- a/Engine/Plugins/Runtime/OpenXRHandTracking/Source/OpenXRHandTracking/Private/OpenXRHandTracking.cpp
+++ b/Engine/Plugins/Runtime/OpenXRHandTracking/Source/OpenXRHandTracking/Private/OpenXRHandTracking.cpp
@@ -387,60 +387,25 @@ bool FOpenXRHandTracking::GetControllerOrientationAndPosition(const int32 Contro
 		OutOrientation = ControllerTransform.GetRotation().Rotator();
 	}
 
-	// Then call super to handle a few of the default labels (eg AnyHand), for backward compatibility
-	if (bSupportLegacyControllerMotionSources)
-	{
-		FXRMotionControllerBase::GetControllerOrientationAndPosition(ControllerIndex, MotionSource, OutOrientation, OutPosition, WorldToMetersScale);
-	}
-
 	return bTracked;
 }
 
-bool FOpenXRHandTracking::GetControllerOrientationAndPosition(const int32 ControllerIndex, const EControllerHand DeviceHand, FRotator& OutOrientation, FVector& OutPosition, float WorldToMetersScale) const
+ETrackingStatus FOpenXRHandTracking::GetControllerTrackingStatus(const int32 ControllerIndex, const FName MotionSource) const
 {
 	if (!bHandTrackingAvailable)
 	{
-		return false;
-	}
-
-	bool bControllerTracked = false;
-	if (ControllerIndex == DeviceIndex)
-	{
-		if (GetControllerTrackingStatus(ControllerIndex, DeviceHand) != ETrackingStatus::NotTracked)
-		{
-			const FTransform* ControllerTransform = nullptr;
-
-			if (DeviceHand == EControllerHand::Left)
-			{
-				ControllerTransform = &GetLeftHandState().GetTransform(EHandKeypoint::Palm);
-			}
-			else if (DeviceHand == EControllerHand::Right)
-			{
-				ControllerTransform = &GetRightHandState().GetTransform(EHandKeypoint::Palm);
-			}
-
-			if (ControllerTransform != nullptr)
-			{
-				OutPosition = ControllerTransform->GetLocation();
-				OutOrientation = ControllerTransform->GetRotation().Rotator();
-
-				bControllerTracked = true;
-			}
-		}
+		return ETrackingStatus::NotTracked;
 	}
 
-	return bControllerTracked;
-}
-
-ETrackingStatus FOpenXRHandTracking::GetControllerTrackingStatus(const int32 ControllerIndex, const EControllerHand DeviceHand) const
-{
-	if (!bHandTrackingAvailable)
+	const MotionSourceInfo* KeyPointInfoPtr = MotionSourceToKeypointMap.Find(MotionSource);
+	if (KeyPointInfoPtr)
 	{
-		return ETrackingStatus::NotTracked;
+		const bool bIsLeft = KeyPointInfoPtr->Value;
+		const FOpenXRHandTracking::FHandState& HandState = bIsLeft ? GetLeftHandState() : GetRightHandState();
+		return HandState.ReceivedJointPoses ? ETrackingStatus::Tracked : ETrackingStatus::NotTracked;
 	}
 
-	const FOpenXRHandTracking::FHandState& HandState = (DeviceHand == EControllerHand::Left) ? GetLeftHandState() : GetRightHandState();
-	return HandState.ReceivedJointPoses ? ETrackingStatus::Tracked : ETrackingStatus::NotTracked;
+	return ETrackingStatus::NotTracked;
 }
 
 FName FOpenXRHandTracking::GetMotionControllerDeviceTypeName() const
diff --git a/Engine/Plugins/Runtime/OpenXRHandTracking/Source/OpenXRHandTracking/Private/OpenXRHandTracking.h b/Engine/Plugins/Runtime/OpenXRHandTracking/Source/OpenXRHandTracking/Private/OpenXRHandTracking.h
--- a/Engine/Plugins/Runtime/OpenXRHandTracking/Source/OpenXRHandTracking/Private/OpenXRHandTracking.h
+++ b/Engine/Plugins/Runtime/OpenXRHandTracking/Source/OpenXRHandTracking/Private/OpenXRHandTracking.h
@@ -60,8 +60,7 @@ public:
 
 	/** IMotionController interface */
 	virtual bool GetControllerOrientationAndPosition(const int32 ControllerIndex, const FName MotionSource, FRotator& OutOrientation, FVector& OutPosition, float WorldToMetersScale) const override;
-	virtual bool GetControllerOrientationAndPosition(const int32 ControllerIndex, const EControllerHand DeviceHand, FRotator& OutOrientation, FVector& OutPosition, float WorldToMetersScale) const override;
-	virtual ETrackingStatus GetControllerTrackingStatus(const int32 ControllerIndex, const EControllerHand DeviceHand) const override;
+	virtual ETrackingStatus GetControllerTrackingStatus(const int32 ControllerIndex, const FName MotionSource) const override;
 	virtual FName GetMotionControllerDeviceTypeName() const override;
 	virtual void EnumerateSources(TArray<FMotionControllerSource>& SourcesOut) const override;
 
diff --git a/Engine/Plugins/Runtime/OpenXRViveTracker/Source/OpenXRViveTracker/Private/OpenXRViveTracker.cpp b/Engine/Plugins/Runtime/OpenXRViveTracker/Source/OpenXRViveTracker/Private/OpenXRViveTracker.cpp
--- a/Engine/Plugins/Runtime/OpenXRViveTracker/Source/OpenXRViveTracker/Private/OpenXRViveTracker.cpp
+++ b/Engine/Plugins/Runtime/OpenXRViveTracker/Source/OpenXRViveTracker/Private/OpenXRViveTracker.cpp
@@ -279,7 +279,7 @@ bool FOpenXRViveTracker::GetControllerOrientationAndPositionForTime(const int32
 	return false;
 }
 
-ETrackingStatus FOpenXRViveTracker::GetControllerTrackingStatus(const int32 ControllerIndex, const EControllerHand DeviceHand) const
+ETrackingStatus FOpenXRViveTracker::GetControllerTrackingStatus(const int32 ControllerIndex, const FName MotionSource) const
 {
 	if (!bActionsAttached || OpenXRHMD == nullptr)
 	{
@@ -292,6 +292,13 @@ ETrackingStatus FOpenXRViveTracker::GetControllerTrackingStatus(const int32 Cont
 		return ETrackingStatus::NotTracked;
 	}
 
+	if (!(ControllerIndex == DeviceIndex && MotionSourceToEControllerHandMap.Contains(MotionSource)))
+	{
+		return ETrackingStatus::NotTracked;
+	}
+
+	const EControllerHand DeviceHand = MotionSourceToEControllerHandMap.FindChecked(MotionSource);
+
 	const FViveTracker* Tracker = Trackers.Find(DeviceHand);
 	if (!Tracker)
 	{
diff --git a/Engine/Plugins/Runtime/OpenXRViveTracker/Source/OpenXRViveTracker/Private/OpenXRViveTracker.h b/Engine/Plugins/Runtime/OpenXRViveTracker/Source/OpenXRViveTracker/Private/OpenXRViveTracker.h
--- a/Engine/Plugins/Runtime/OpenXRViveTracker/Source/OpenXRViveTracker/Private/OpenXRViveTracker.h
+++ b/Engine/Plugins/Runtime/OpenXRViveTracker/Source/OpenXRViveTracker/Private/OpenXRViveTracker.h
@@ -56,9 +56,8 @@ public:
 
 	/** IMotionController interface */
 	virtual bool GetControllerOrientationAndPosition(const int32 ControllerIndex, const FName MotionSource, FRotator& OutOrientation, FVector& OutPosition, float WorldToMetersScale) const override;
-	virtual bool GetControllerOrientationAndPosition(const int32 ControllerIndex, const EControllerHand DeviceHand, FRotator& OutOrientation, FVector& OutPosition, float WorldToMetersScale) const override { check(false); return false; }
 	virtual bool GetControllerOrientationAndPositionForTime(const int32 ControllerIndex, const FName MotionSource, FTimespan Time, bool& OutTimeWasUsed, FRotator& OutOrientation, FVector& OutPosition, bool& OutbProvidedLinearVelocity, FVector& OutLinearVelocity, bool& OutbProvidedAngularVelocity, FVector& OutAngularVelocityRadPerSec, bool& OutbProvidedLinearAcceleration, FVector& OutLinearAcceleration, float WorldToMetersScale) const override;
-	virtual ETrackingStatus GetControllerTrackingStatus(const int32 ControllerIndex, const EControllerHand DeviceHand) const override;
+	virtual ETrackingStatus GetControllerTrackingStatus(const int32 ControllerIndex, const FName MotionSource) const override;
 	virtual FName GetMotionControllerDeviceTypeName() const override;
 	virtual void EnumerateSources(TArray<FMotionControllerSource>& SourcesOut) const override;
 
diff --git a/Engine/Source/Runtime/HeadMountedDisplay/Private/HeadMountedDisplayFunctionLibrary.cpp b/Engine/Source/Runtime/HeadMountedDisplay/Private/HeadMountedDisplayFunctionLibrary.cpp
--- a/Engine/Source/Runtime/HeadMountedDisplay/Private/HeadMountedDisplayFunctionLibrary.cpp
+++ b/Engine/Source/Runtime/HeadMountedDisplay/Private/HeadMountedDisplayFunctionLibrary.cpp
@@ -10,7 +10,6 @@
 #include "IXRSystemAssets.h"
 #include "Components/PrimitiveComponent.h"
 #include "Features/IModularFeatures.h"
-#include "XRMotionControllerBase.h" // for GetHandEnumForSourceName()
 
 #include UE_INLINE_GENERATED_CPP_BY_NAME(HeadMountedDisplayFunctionLibrary)
 
diff --git a/Engine/Source/Runtime/HeadMountedDisplay/Private/XRMotionControllerBase.cpp b/Engine/Source/Runtime/HeadMountedDisplay/Private/XRMotionControllerBase.cpp
--- a/Engine/Source/Runtime/HeadMountedDisplay/Private/XRMotionControllerBase.cpp
+++ b/Engine/Source/Runtime/HeadMountedDisplay/Private/XRMotionControllerBase.cpp
@@ -33,31 +33,6 @@
 	};
 }
 
-bool FXRMotionControllerBase::GetControllerOrientationAndPosition(const int32 ControllerIndex, const FName MotionSource, FRotator& OutOrientation, FVector& OutPosition, float WorldToMetersScale) const
-{
-	bool bSucess = false;
-	if (ControllerIndex != INDEX_NONE)
-	{
-		EControllerHand DeviceHand;
-		if (GetHandEnumForSourceName(MotionSource, DeviceHand))
-		{
-			if (DeviceHand == EControllerHand::AnyHand)
-			{
-				bSucess = GetControllerOrientationAndPosition(ControllerIndex, EControllerHand::Left, OutOrientation, OutPosition, WorldToMetersScale);
-				if (!bSucess)
-				{
-					bSucess = GetControllerOrientationAndPosition(ControllerIndex, EControllerHand::Right, OutOrientation, OutPosition, WorldToMetersScale);
-				}
-			}
-			else
-			{
-				bSucess = GetControllerOrientationAndPosition(ControllerIndex, DeviceHand, OutOrientation, OutPosition, WorldToMetersScale);
-			}
-		}
-	}
-	return bSucess;
-}
-
 bool FXRMotionControllerBase::GetControllerOrientationAndPositionForTime(const int32 ControllerIndex, const FName MotionSource, FTimespan Time, bool& OutTimeWasUsed, FRotator& OutOrientation, FVector& OutPosition, bool& OutbProvidedLinearVelocity, FVector& OutLinearVelocity, bool& OutbProvidedAngularVelocity, FVector& OutAngularVelocityRadPerSec, bool& OutbProvidedLinearAcceleration, FVector& OutLinearAcceleration, float WorldToMetersScale) const
 {
 	// Default implementation simply ignores the Timespan, no additional accuracy is provided by this call, velocities are not provided.
@@ -68,32 +43,6 @@
 	return GetControllerOrientationAndPosition(ControllerIndex, MotionSource, OutOrientation, OutPosition, WorldToMetersScale);
 }
 
-ETrackingStatus FXRMotionControllerBase::GetControllerTrackingStatus(const int32 ControllerIndex, const FName MotionSource) const
-{
-	if (ControllerIndex != INDEX_NONE)
-	{
-		EControllerHand DeviceHand;
-		if (GetHandEnumForSourceName(MotionSource, DeviceHand))
-		{
-			if (DeviceHand == EControllerHand::AnyHand)
-			{
-				FRotator ThrowawayOrientation;
-				FVector  ThrowawayPosition;
-				// we've moved explicit handling of the 'AnyHand' source into this class from UMotionControllerComponent,
-				// to maintain behavior we use the return value of GetControllerOrientationAndPosition() to determine which hand's
-				// status we should check
-				if (GetControllerOrientationAndPosition(ControllerIndex, EControllerHand::Left, ThrowawayOrientation, ThrowawayPosition, 100.f))
-				{
-					return GetControllerTrackingStatus(ControllerIndex, EControllerHand::Left);
-				}
-				return GetControllerTrackingStatus(ControllerIndex, EControllerHand::Right);
-			}
-			return GetControllerTrackingStatus(ControllerIndex, DeviceHand);
-		}
-	}
-	return ETrackingStatus::NotTracked;
-}
-
 void FXRMotionControllerBase::EnumerateSources(TArray<FMotionControllerSource>& SourcesOut) const
 {
 	const int32 LegaceSourcesCount = UE_ARRAY_COUNT(XRMotionControllerBase_Impl::LegacySources);
@@ -164,3 +113,55 @@
 		return false;
 	}
 }
+
+
+bool FXRMotionControllerBaseLegacy::GetControllerOrientationAndPosition(const int32 ControllerIndex, const FName MotionSource, FRotator& OutOrientation, FVector& OutPosition, float WorldToMetersScale) const
+{
+	bool bSucess = false;
+	if (ControllerIndex != INDEX_NONE)
+	{
+		EControllerHand DeviceHand;
+		if (GetHandEnumForSourceName(MotionSource, DeviceHand))
+		{
+			if (DeviceHand == EControllerHand::AnyHand)
+			{
+				bSucess = GetControllerOrientationAndPosition(ControllerIndex, EControllerHand::Left, OutOrientation, OutPosition, WorldToMetersScale);
+				if (!bSucess)
+				{
+					bSucess = GetControllerOrientationAndPosition(ControllerIndex, EControllerHand::Right, OutOrientation, OutPosition, WorldToMetersScale);
+				}
+			}
+			else
+			{
+				bSucess = GetControllerOrientationAndPosition(ControllerIndex, DeviceHand, OutOrientation, OutPosition, WorldToMetersScale);
+			}
+		}
+	}
+	return bSucess;
+}
+
+ETrackingStatus FXRMotionControllerBaseLegacy::GetControllerTrackingStatus(const int32 ControllerIndex, const FName MotionSource) const
+{
+	if (ControllerIndex != INDEX_NONE)
+	{
+		EControllerHand DeviceHand;
+		if (GetHandEnumForSourceName(MotionSource, DeviceHand))
+		{
+			if (DeviceHand == EControllerHand::AnyHand)
+			{
+				FRotator ThrowawayOrientation;
+				FVector  ThrowawayPosition;
+				// we've moved explicit handling of the 'AnyHand' source into this class from UMotionControllerComponent,
+				// to maintain behavior we use the return value of GetControllerOrientationAndPosition() to determine which hand's
+				// status we should check
+				if (GetControllerOrientationAndPosition(ControllerIndex, EControllerHand::Left, ThrowawayOrientation, ThrowawayPosition, 100.f))
+				{
+					return GetControllerTrackingStatus(ControllerIndex, EControllerHand::Left);
+				}
+				return GetControllerTrackingStatus(ControllerIndex, EControllerHand::Right);
+			}
+			return GetControllerTrackingStatus(ControllerIndex, DeviceHand);
+		}
+	}
+	return ETrackingStatus::NotTracked;
+}
diff --git a/Engine/Source/Runtime/HeadMountedDisplay/Public/XRMotionControllerBase.h b/Engine/Source/Runtime/HeadMountedDisplay/Public/XRMotionControllerBase.h
--- a/Engine/Source/Runtime/HeadMountedDisplay/Public/XRMotionControllerBase.h
+++ b/Engine/Source/Runtime/HeadMountedDisplay/Public/XRMotionControllerBase.h
@@ -20,9 +20,9 @@ public:
 	virtual ~FXRMotionControllerBase() {};
 
 	// Begin IMotionController interface
-	virtual bool GetControllerOrientationAndPosition(const int32 ControllerIndex, const FName MotionSource, FRotator& OutOrientation, FVector& OutPosition, float WorldToMetersScale) const;
+	virtual bool GetControllerOrientationAndPosition(const int32 ControllerIndex, const FName MotionSource, FRotator& OutOrientation, FVector& OutPosition, float WorldToMetersScale) const = 0;
 	virtual bool GetControllerOrientationAndPositionForTime(const int32 ControllerIndex, const FName MotionSource, FTimespan Time, bool& OutTimeWasUsed, FRotator& OutOrientation, FVector& OutPosition, bool& OutbProvidedLinearVelocity, FVector& OutLinearVelocity, bool& OutbProvidedAngularVelocity, FVector& OutAngularVelocityRadPerSec, bool& OutbProvidedLinearAcceleration, FVector& OutLinearAcceleration, float WorldToMetersScale) const;
-	virtual ETrackingStatus GetControllerTrackingStatus(const int32 ControllerIndex, const FName MotionSource) const;
+	virtual ETrackingStatus GetControllerTrackingStatus(const int32 ControllerIndex, const FName MotionSource) const = 0;
 	virtual void EnumerateSources(TArray<FMotionControllerSource>& SourcesOut) const;
 	virtual float GetCustomParameterValue(const FName MotionSource, FName ParameterName, bool& bValueFound) const { bValueFound = false;  return 0.f; }
 	virtual bool GetHandJointPosition(const FName MotionSource, int jointIndex, FVector& OutPosition) const override { return false; }
@@ -34,8 +34,21 @@ public:
 	static FName HMDSourceId;
 
 	static bool GetHandEnumForSourceName(const FName Source, EControllerHand& OutHand);
+};
 
-	// Original GetControllerOrientationAndPosition signature for backwards compatibility 
+// This class adapts deprecated or maintenence only xr plugins which use the old EControllerHand motion sources to the newer FName MotionSource API.
+// If a plugin is still in active development it should refactor and derive from IMotionController or FXRMotionControllerBase rather than derive from this class.
+class HEADMOUNTEDDISPLAY_API FXRMotionControllerBaseLegacy : public FXRMotionControllerBase
+{
+private:
+	// Begin IMotionController interface
+	// These are overridden to call into the backward compatibility functions below.  They are final and private because a legacy vr plugin would not use or implement these functions.
+	virtual bool GetControllerOrientationAndPosition(const int32 ControllerIndex, const FName MotionSource, FRotator& OutOrientation, FVector& OutPosition, float WorldToMetersScale) const override final;
+	virtual ETrackingStatus GetControllerTrackingStatus(const int32 ControllerIndex, const FName MotionSource) const override final;
+	// End IMotionController interface
+
+public:
+	// Original GetControllerOrientationAndPosition signature for backwards compatibility.  They are pure virtual because a legacy plugin must implement them and public because a legacy vr plugin could use them internally.
 	virtual bool GetControllerOrientationAndPosition(const int32 ControllerIndex, const EControllerHand DeviceHand, FRotator& OutOrientation, FVector& OutPosition, float WorldToMetersScale) const = 0;
 	// Original GetControllerTrackingStatus signature for backwards compatibility
 	virtual ETrackingStatus GetControllerTrackingStatus(const int32 ControllerIndex, const EControllerHand DeviceHand) const = 0;
