Unify setting readback of bStartInVR into a single helper function in IStereoRendering.

diff --git a/Engine/Plugins/Runtime/OpenXR/Source/OpenXRHMD/Private/OpenXRHMD.cpp b/Engine/Plugins/Runtime/OpenXR/Source/OpenXRHMD/Private/OpenXRHMD.cpp
--- a/Engine/Plugins/Runtime/OpenXR/Source/OpenXRHMD/Private/OpenXRHMD.cpp
+++ b/Engine/Plugins/Runtime/OpenXR/Source/OpenXRHMD/Private/OpenXRHMD.cpp
@@ -1290,9 +1290,7 @@ FOpenXRHMD::FOpenXRHMD(const FAutoRegister& AutoRegister, XrInstance InInstance,
 	ReconfigureForShaderPlatform(GMaxRHIShaderPlatform);
 
 #if PLATFORM_HOLOLENS || PLATFORM_ANDROID
-	bool bStartInVR = false;
-	GConfig->GetBool(TEXT("/Script/EngineSettings.GeneralProjectSettings"), TEXT("bStartInVR"), bStartInVR, GGameIni);
-	bIsStandaloneStereoOnlyDevice = FParse::Param(FCommandLine::Get(), TEXT("vr")) || bStartInVR;
+	bIsStandaloneStereoOnlyDevice = IStereoRendering::IsStartInVR();
 #else
 	for (IOpenXRExtensionPlugin* Module : ExtensionPlugins)
 	{
diff --git a/Engine/Plugins/Runtime/OpenXR/Source/OpenXRHMD/Private/OpenXRHMDModule.cpp b/Engine/Plugins/Runtime/OpenXR/Source/OpenXRHMD/Private/OpenXRHMDModule.cpp
--- a/Engine/Plugins/Runtime/OpenXR/Source/OpenXRHMD/Private/OpenXRHMDModule.cpp
+++ b/Engine/Plugins/Runtime/OpenXR/Source/OpenXRHMD/Private/OpenXRHMDModule.cpp
@@ -179,9 +179,7 @@ bool FOpenXRHMDModule::IsStandaloneStereoOnlyDevice()
 		}
 
 #if PLATFORM_HOLOLENS || PLATFORM_ANDROID
-		bool bStartInVR = false;
-		GConfig->GetBool(TEXT("/Script/EngineSettings.GeneralProjectSettings"), TEXT("bStartInVR"), bStartInVR, GGameIni); 
-		return FParse::Param(FCommandLine::Get(), TEXT("vr")) || bStartInVR;
+		return IStereoRendering::IsStartInVR();
 #endif
 	}
 	return false;
diff --git a/Engine/Source/Runtime/Engine/Private/AudioDevice.cpp b/Engine/Source/Runtime/Engine/Private/AudioDevice.cpp
--- a/Engine/Source/Runtime/Engine/Private/AudioDevice.cpp
+++ b/Engine/Source/Runtime/Engine/Private/AudioDevice.cpp
@@ -28,6 +28,7 @@
 #include "Stats/StatsTrace.h"
 #include "UnrealEngine.h"
 #include "UObject/UObjectIterator.h"
+#include "StereoRendering.h"
 
 #if WITH_EDITOR
 #include "AssetRegistry/AssetRegistryModule.h"
@@ -7219,7 +7220,7 @@ bool FAudioDevice::CanUseVRAudioDevice()
 	else
 #endif
 	{
-		return FParse::Param(FCommandLine::Get(), TEXT("vr")) || GetDefault<UGeneralProjectSettings>()->bStartInVR;
+		return IStereoRendering::IsStartInVR();
 	}
 }
 
diff --git a/Engine/Source/Runtime/Engine/Private/StereoRendering.cpp b/Engine/Source/Runtime/Engine/Private/StereoRendering.cpp
--- a/Engine/Source/Runtime/Engine/Private/StereoRendering.cpp
+++ b/Engine/Source/Runtime/Engine/Private/StereoRendering.cpp
@@ -2,6 +2,8 @@
 
 #include "StereoRendering.h"
 #include "SceneView.h"
+#include "GeneralProjectSettings.h"
+#include "Misc/ConfigCacheIni.h"
 
 
 bool IStereoRendering::IsStereoEyeView(const FSceneView& View)
@@ -18,3 +20,19 @@ bool IStereoRendering::IsASecondaryView(const FSceneView& View)
 {
 	return IsASecondaryPass(View.StereoPass);
 }
+
+bool IStereoRendering::IsStartInVR()
+{
+	bool bStartInVR = false;
+
+	if (IsClassLoaded<UGeneralProjectSettings>())
+	{
+		bStartInVR = GetDefault<UGeneralProjectSettings>()->bStartInVR;
+	}
+	else
+	{
+		GConfig->GetBool(TEXT("/Script/EngineSettings.GeneralProjectSettings"), TEXT("bStartInVR"), bStartInVR, GGameIni);
+	}
+
+	return FParse::Param(FCommandLine::Get(), TEXT("vr")) || bStartInVR;
+}
diff --git a/Engine/Source/Runtime/Engine/Private/UnrealEngine.cpp b/Engine/Source/Runtime/Engine/Private/UnrealEngine.cpp
--- a/Engine/Source/Runtime/Engine/Private/UnrealEngine.cpp
+++ b/Engine/Source/Runtime/Engine/Private/UnrealEngine.cpp
@@ -3786,7 +3786,7 @@ bool UEngine::InitializeHMDDevice()
 				StereoRenderingDevice = XRSystem->GetStereoRenderingDevice();
 				const bool bShouldStartInVR = StereoRenderingDevice.IsValid() && 
 											  !GIsEditor &&
-											  (FParse::Param(FCommandLine::Get(), TEXT("vr")) || GetDefault<UGeneralProjectSettings>()->bStartInVR);
+											  IStereoRendering::IsStartInVR();
 				if (bShouldStartInVR)
 				{
 					StereoRenderingDevice->EnableStereo(true);
diff --git a/Engine/Source/Runtime/Engine/Public/StereoRendering.h b/Engine/Source/Runtime/Engine/Public/StereoRendering.h
--- a/Engine/Source/Runtime/Engine/Public/StereoRendering.h
+++ b/Engine/Source/Runtime/Engine/Public/StereoRendering.h
@@ -183,4 +183,8 @@ public:
 	virtual bool OverrideFinalPostprocessSettings(struct FPostProcessSettings* OverridePostProcessingSettings, const enum EStereoscopicPass StereoPassType, const int32 StereoViewIndex, float& BlendWeight) { return false; }
 	virtual void EndFinalPostprocessSettings(struct FPostProcessSettings* FinalPostProcessingSettings, const enum EStereoscopicPass StereoPassType, const int32 StereoViewIndex) {}
 
+	/**
+	* Static helper. Return true if the Start in VR setting is set to true, or if we have the "-vr" commandline argument
+	*/
+	static bool IsStartInVR();
 };
diff --git a/Engine/Source/Runtime/HeadMountedDisplay/Private/HeadMountedDisplayVulkanExtensions.cpp b/Engine/Source/Runtime/HeadMountedDisplay/Private/HeadMountedDisplayVulkanExtensions.cpp
--- a/Engine/Source/Runtime/HeadMountedDisplay/Private/HeadMountedDisplayVulkanExtensions.cpp
+++ b/Engine/Source/Runtime/HeadMountedDisplay/Private/HeadMountedDisplayVulkanExtensions.cpp
@@ -3,11 +3,9 @@
 #include "IHeadMountedDisplayVulkanExtensions.h"
 #include "GeneralProjectSettings.h"
 #include "Misc/CommandLine.h"
+#include "StereoRendering.h"
 
 bool IHeadMountedDisplayVulkanExtensions::ShouldDisableVulkanVSync() const
 {
-	// This can be called while setting up a vulkan swap chain for the PreLoadScreenManager but we cannot easily inspect bStartInVR then because the object has not been loaded (it being PreLoad).
-	// For now we are just going to return false here and avoid an assert when trying to GetDefault.  This might break VR rendering of PreLoadScreens... but there is a good chance it is already broken.
-	// We check for the class being loaded first so that -vr behavior is consistent with bStartInVR.
-	return IsClassLoaded<UGeneralProjectSettings>() && (FParse::Param(FCommandLine::Get(), TEXT("vr")) || GetDefault<UGeneralProjectSettings>()->bStartInVR);
+	return IStereoRendering::IsStartInVR();
 }
