Wrap IMotionController.h with deprecation warning disable pragma for reference to the now deprecated UPlayerMappableInputConfig.

#rnx

diff --git a/Engine/Source/Runtime/HeadMountedDisplay/Public/IMotionController.h b/Engine/Source/Runtime/HeadMountedDisplay/Public/IMotionController.h
--- a/Engine/Source/Runtime/HeadMountedDisplay/Public/IMotionController.h
+++ b/Engine/Source/Runtime/HeadMountedDisplay/Public/IMotionController.h
@@ -125,6 +125,7 @@ public:
 	
 	virtual bool GetHandJointPosition(const FName MotionSource, int jointIndex, FVector& OutPosition) const = 0;
 
+	PRAGMA_DISABLE_DEPRECATION_WARNINGS
 	/**
 	 * Add a player mappable input config to the motion controller. This allows the motion controller to support
 	 * Enhanced Input actions.
@@ -134,6 +135,7 @@ public:
 	 * @return			False if the input config can't be attached to the session, true otherwise
 	 */
 	virtual bool SetPlayerMappableInputConfig(TObjectPtr<class UPlayerMappableInputConfig> InputConfig = nullptr) { return true; };
+	PRAGMA_ENABLE_DEPRECATION_WARNINGS
 
 	// explicit source names
 	static HEADMOUNTEDDISPLAY_API FName LeftHandSourceId;
