UE-192555 MotionControllerComponent difficult to extend.
-Changing access of much of the implementation from private to protected.
-It is currently possible to replace all of that implementation because the incoming calls are public and virtual but this leads to a lot of duplication and replacement of code that would otherwise work.

diff --git a/Engine/Source/Runtime/HeadMountedDisplay/Public/MotionControllerComponent.h b/Engine/Source/Runtime/HeadMountedDisplay/Public/MotionControllerComponent.h
--- a/Engine/Source/Runtime/HeadMountedDisplay/Public/MotionControllerComponent.h
+++ b/Engine/Source/Runtime/HeadMountedDisplay/Public/MotionControllerComponent.h
@@ -170,7 +170,7 @@ protected:
 	HEADMOUNTEDDISPLAY_API bool GetLinearAcceleration(FVector& OutLinearAcceleration) const;
 
 
-private:
+protected:
 
 	/** Whether or not this component had a valid tracked controller associated with it this frame*/
 	bool bTracked;
@@ -222,7 +222,7 @@ private:
 		virtual int32 GetPriority() const override { return -10; }
 		virtual bool IsActiveThisFrame_Internal(const FSceneViewExtensionContext& Context) const;
 
-	private:
+	protected:
 		friend class UMotionControllerComponent;
 
 		/** Motion controller component associated with this view extension */
