Fix remaining PRAGMA_DISABLE_DEPRECATION_WARNINGS without matching PRAGMA_ENABLE_DEPRECATION_WARNINGS

diff --git a/Engine/Source/Runtime/HeadMountedDisplay/Private/MotionControllerComponent.cpp b/Engine/Source/Runtime/HeadMountedDisplay/Private/MotionControllerComponent.cpp
--- a/Engine/Source/Runtime/HeadMountedDisplay/Private/MotionControllerComponent.cpp
+++ b/Engine/Source/Runtime/HeadMountedDisplay/Private/MotionControllerComponent.cpp
@@ -49,8 +49,6 @@ namespace LegacyMotionSources
 	}
 }
 
-PRAGMA_DISABLE_DEPRECATION_WARNINGS
-
 //=============================================================================
 UMotionControllerComponent::UMotionControllerComponent(const FObjectInitializer& ObjectInitializer)
 	: Super(ObjectInitializer)
