UE-205065 XRMotionController component when set to HMD produces log spam
-Do the special case code for the hmd before trying the more generic code for controller poses because we log an error if we try the controller code with the HMD motion source.
-Added "Head" to the motion sources provided by OpenXR, so it shows up in the motion controller dropdown.  Made "Head" and "HMD" function identically, "HMD" is the legacy term, "Head" the more modern.

diff --git a/Engine/Plugins/Runtime/OpenXR/Source/OpenXRInput/Private/OpenXRInput.cpp b/Engine/Plugins/Runtime/OpenXR/Source/OpenXRInput/Private/OpenXRInput.cpp
--- a/Engine/Plugins/Runtime/OpenXR/Source/OpenXRInput/Private/OpenXRInput.cpp
+++ b/Engine/Plugins/Runtime/OpenXR/Source/OpenXRInput/Private/OpenXRInput.cpp
@@ -32,6 +32,7 @@
 
 namespace OpenXRSourceNames
 {
+	static const FName Head("Head");
 	static const FName AnyHand("AnyHand");
 	static const FName Left("Left");
 	static const FName Right("Right");
@@ -1357,6 +1358,7 @@ void FOpenXRInputPlugin::FOpenXRInput::EnumerateSources(TArray<FMotionController
 {
 	check(IsInGameThread());
 
+	SourcesOut.Add(OpenXRSourceNames::Head);
 	SourcesOut.Add(OpenXRSourceNames::AnyHand);
 	SourcesOut.Add(OpenXRSourceNames::Left);
 	SourcesOut.Add(OpenXRSourceNames::Right);
diff --git a/Engine/Source/Runtime/HeadMountedDisplay/Private/IMotionController.cpp b/Engine/Source/Runtime/HeadMountedDisplay/Private/IMotionController.cpp
--- a/Engine/Source/Runtime/HeadMountedDisplay/Private/IMotionController.cpp
+++ b/Engine/Source/Runtime/HeadMountedDisplay/Private/IMotionController.cpp
@@ -6,6 +6,7 @@
 FName IMotionController::LeftHandSourceId(TEXT("Left"));
 FName IMotionController::RightHandSourceId(TEXT("Right"));
 FName IMotionController::HMDSourceId(TEXT("HMD"));
+FName IMotionController::HeadSourceId(TEXT("Head"));
 
 bool IMotionController::GetHandEnumForSourceName(const FName Source, EControllerHand& OutHand)
 {
@@ -57,6 +58,7 @@ bool IMotionController::GetHandEnumForSourceName(const FName Source, EController
 		MotionSourceToEControllerHandMap.Add(TEXT("RightGrip"), EControllerHand::Right);
 		MotionSourceToEControllerHandMap.Add(TEXT("LeftAim"), EControllerHand::Left);
 		MotionSourceToEControllerHandMap.Add(TEXT("RightAim"), EControllerHand::Right);
+		MotionSourceToEControllerHandMap.Add(IMotionController::HeadSourceId, EControllerHand::HMD);
 	}
 
 	EControllerHand* FoundEnum = MotionSourceToEControllerHandMap.Find(Source);
diff --git a/Engine/Source/Runtime/HeadMountedDisplay/Private/MotionControllerComponent.cpp b/Engine/Source/Runtime/HeadMountedDisplay/Private/MotionControllerComponent.cpp
--- a/Engine/Source/Runtime/HeadMountedDisplay/Private/MotionControllerComponent.cpp
+++ b/Engine/Source/Runtime/HeadMountedDisplay/Private/MotionControllerComponent.cpp
@@ -305,31 +305,7 @@ bool UMotionControllerComponent::PollControllerState_GameThread(FVector& Positio
 			bPolledHMD_GameThread = false;
 		}
 
-		TArray<IMotionController*> MotionControllers;
-		MotionControllers = IModularFeatures::Get().GetModularFeatureImplementations<IMotionController>(IMotionController::GetModularFeatureName());
-		for (auto MotionController : MotionControllers)
-		{
-			if (MotionController == nullptr)
-			{
-				continue;
-			}
-
-			CurrentTrackingStatus = MotionController->GetControllerTrackingStatus(PlayerIndex, MotionSource);
-			if (MotionController->GetControllerOrientationAndPosition(PlayerIndex, MotionSource, Orientation, Position, OutbProvidedLinearVelocity, OutLinearVelocity, OutbProvidedAngularVelocity, OutAngularVelocityAsAxisAndLength, OutbProvidedLinearAcceleration, OutLinearAcceleration, WorldToMetersScale))
-			{
-				InUseMotionController = MotionController;
-				OnMotionControllerUpdated();
-				InUseMotionController = nullptr;
-
-				{
-					FScopeLock Lock(&PolledMotionControllerMutex);
-					PolledMotionController_GameThread = MotionController;  // We only want a render thread update from the motion controller we polled on the game thread.
-				}
-				return true;
-			}
-		}
-
-		if (MotionSource == IMotionController::HMDSourceId)
+		if (MotionSource == IMotionController::HMDSourceId || MotionSource == IMotionController::HeadSourceId)
 		{
 			IXRTrackingSystem* TrackingSys = GEngine->XRSystem.Get();
 			if (TrackingSys)
@@ -346,6 +322,32 @@ bool UMotionControllerComponent::PollControllerState_GameThread(FVector& Positio
 				}
 			}
 		}
+		else
+		{
+			TArray<IMotionController*> MotionControllers;
+			MotionControllers = IModularFeatures::Get().GetModularFeatureImplementations<IMotionController>(IMotionController::GetModularFeatureName());
+			for (auto MotionController : MotionControllers)
+			{
+				if (MotionController == nullptr)
+				{
+					continue;
+				}
+
+				CurrentTrackingStatus = MotionController->GetControllerTrackingStatus(PlayerIndex, MotionSource);
+				if (MotionController->GetControllerOrientationAndPosition(PlayerIndex, MotionSource, Orientation, Position, OutbProvidedLinearVelocity, OutLinearVelocity, OutbProvidedAngularVelocity, OutAngularVelocityAsAxisAndLength, OutbProvidedLinearAcceleration, OutLinearAcceleration, WorldToMetersScale))
+				{
+					InUseMotionController = MotionController;
+					OnMotionControllerUpdated();
+					InUseMotionController = nullptr;
+
+					{
+						FScopeLock Lock(&PolledMotionControllerMutex);
+						PolledMotionController_GameThread = MotionController;  // We only want a render thread update from the motion controller we polled on the game thread.
+					}
+					return true;
+				}
+			}
+		}
 	}
 	return false;
 }
diff --git a/Engine/Source/Runtime/HeadMountedDisplay/Public/IMotionController.h b/Engine/Source/Runtime/HeadMountedDisplay/Public/IMotionController.h
--- a/Engine/Source/Runtime/HeadMountedDisplay/Public/IMotionController.h
+++ b/Engine/Source/Runtime/HeadMountedDisplay/Public/IMotionController.h
@@ -178,6 +178,7 @@ public:
 	static HEADMOUNTEDDISPLAY_API FName LeftHandSourceId;
 	static HEADMOUNTEDDISPLAY_API FName RightHandSourceId;
 	static HEADMOUNTEDDISPLAY_API FName HMDSourceId;
+	static HEADMOUNTEDDISPLAY_API FName HeadSourceId; // Note Head and HMD are interchangable.  HMD is the legacy term.
 
 	static HEADMOUNTEDDISPLAY_API bool GetHandEnumForSourceName(const FName Source, EControllerHand& OutHand);
 };
