Error - StructProperty FXRMotionControllerState::ControllerLocation is not initialized
-Initialized the members mentioned.

diff --git a/Engine/Source/Runtime/HeadMountedDisplay/Public/HeadMountedDisplayTypes.h b/Engine/Source/Runtime/HeadMountedDisplay/Public/HeadMountedDisplayTypes.h
--- a/Engine/Source/Runtime/HeadMountedDisplay/Public/HeadMountedDisplayTypes.h
+++ b/Engine/Source/Runtime/HeadMountedDisplay/Public/HeadMountedDisplayTypes.h
@@ -443,16 +443,16 @@ struct FXRMotionControllerState
 	EXRControllerPoseType XRControllerPoseType = EXRControllerPoseType::Grip;
 
 	UPROPERTY(BlueprintReadOnly, Category = "XR")
-	FVector ControllerLocation;
+	FVector ControllerLocation = FVector(0.0f);
 
 	UPROPERTY(BlueprintReadOnly, Category = "XR")
-	FQuat ControllerRotation;
+	FQuat ControllerRotation = FQuat(EForceInit::ForceInitToZero);;
 
 	// These are used by the XRVisualizationFunctionLibrary, and are not exposed to blueprint.
 	UPROPERTY()
-	FVector GripUnrealSpaceLocation;
+	FVector GripUnrealSpaceLocation = FVector(0.0f);
 	UPROPERTY()
-	FQuat GripUnrealSpaceRotation;
+	FQuat GripUnrealSpaceRotation = FQuat(EForceInit::ForceInitToZero);;
 };
 
 USTRUCT(BlueprintType)
