[Engine]
* Readded dllexport in IXRTrackingSystem

diff --git a/Engine/Source/Runtime/HeadMountedDisplay/Public/IXRTrackingSystem.h b/Engine/Source/Runtime/HeadMountedDisplay/Public/IXRTrackingSystem.h
--- a/Engine/Source/Runtime/HeadMountedDisplay/Public/IXRTrackingSystem.h
+++ b/Engine/Source/Runtime/HeadMountedDisplay/Public/IXRTrackingSystem.h
@@ -67,7 +67,7 @@ public:
 	 * Other devices can have arbitrary ids defined by each system.
 	 * If a tracking system does not support tracking HMDs, device ID zero should be treated as invalid.
 	 */
-	static const int32 HMDDeviceId = 0;
+	static HEADMOUNTEDDISPLAY_API const int32 HMDDeviceId = 0;
 
 	/**
 	 * Whether or not the system supports positional tracking (either via sensor or other means)
