Properly pipeline late update states from game thread to render thread in FLateUpdateManager.
Previously, game and render threads would "flip" through a two-element array of states. In cases where the game thread got far ahead of the render thread, this could lead to a race condition.

diff --git a/Engine/Source/Runtime/HeadMountedDisplay/Private/LateUpdateManager.cpp b/Engine/Source/Runtime/HeadMountedDisplay/Private/LateUpdateManager.cpp
--- a/Engine/Source/Runtime/HeadMountedDisplay/Private/LateUpdateManager.cpp
+++ b/Engine/Source/Runtime/HeadMountedDisplay/Private/LateUpdateManager.cpp
@@ -13,51 +13,41 @@ static TAutoConsoleVariable<bool> CVarXRLateUpdateMangerDisable(
 	TEXT("Disable the LateUpdateManager preventing child components from receiving late updates.\n"),
 	ECVF_Default);
 
-FLateUpdateManager::FLateUpdateManager()
-	: LateUpdateGameWriteIndex(0)
-	, LateUpdateRenderReadIndex(0)
-{
-}
-
 void FLateUpdateManager::Setup(const FTransform& ParentToWorld, USceneComponent* Component, bool bSkipLateUpdate)
 {
 	check(IsInGameThread());
 	bSkipLateUpdate = bSkipLateUpdate || CVarXRLateUpdateMangerDisable.GetValueOnGameThread();
 
-	UpdateStates[LateUpdateGameWriteIndex].Primitives.Reset();
-	UpdateStates[LateUpdateGameWriteIndex].ParentToWorld = ParentToWorld;
+	PipelinedUpdateStatesGame.Primitives.Reset();
+	PipelinedUpdateStatesGame.ParentToWorld = ParentToWorld;
 	GatherLateUpdatePrimitives(Component);
-	UpdateStates[LateUpdateGameWriteIndex].bSkip = bSkipLateUpdate;
-	++UpdateStates[LateUpdateGameWriteIndex].TrackingNumber;
+	PipelinedUpdateStatesGame.bSkip = bSkipLateUpdate;
 
-	int32 NextFrameRenderReadIndex = LateUpdateGameWriteIndex;
-	LateUpdateGameWriteIndex = 1 - LateUpdateGameWriteIndex;
-
-	ENQUEUE_RENDER_COMMAND(UpdateLateUpdateRenderReadIndexCommand)(
-		[NextFrameRenderReadIndex, this](FRHICommandListImmediate& RHICmdList)
-	{
-		LateUpdateRenderReadIndex = NextFrameRenderReadIndex;
-	});
+	ENQUEUE_RENDER_COMMAND(UpdateLateUpdateStatesRendering)(
+		[UpdateStatesGame = PipelinedUpdateStatesGame, this](FRHICommandListImmediate& RHICmdList) mutable
+		{
+			PipelinedUpdateStatesRendering = MoveTemp(UpdateStatesGame);
+		});
 }
 
 void FLateUpdateManager::Apply_RenderThread(FSceneInterface* Scene, const FTransform& OldRelativeTransform, const FTransform& NewRelativeTransform)
 {
 	FRHICommandListBase& RHICmdList = FRHICommandListImmediate::Get();
 
-	if (!UpdateStates[LateUpdateRenderReadIndex].Primitives.Num() || UpdateStates[LateUpdateRenderReadIndex].bSkip)
+	if (!PipelinedUpdateStatesRendering.Primitives.Num() || PipelinedUpdateStatesRendering.bSkip)
 	{
 		return;
 	}
 
-	const FTransform OldCameraTransform = OldRelativeTransform * UpdateStates[LateUpdateRenderReadIndex].ParentToWorld;
-	const FTransform NewCameraTransform = NewRelativeTransform * UpdateStates[LateUpdateRenderReadIndex].ParentToWorld;
+	const FTransform OldCameraTransform = OldRelativeTransform * PipelinedUpdateStatesRendering.ParentToWorld;
+	const FTransform NewCameraTransform = NewRelativeTransform * PipelinedUpdateStatesRendering.ParentToWorld;
 	const FMatrix LateUpdateTransform = (OldCameraTransform.Inverse() * NewCameraTransform).ToMatrixWithScale();
 
 	bool bIndicesHaveChanged = false;
 
 	// Apply delta to the cached scene proxies
 	// Also check whether any primitive indices have changed, in case the scene has been modified in the meantime.
-	for (auto& PrimitivePair : UpdateStates[LateUpdateRenderReadIndex].Primitives)
+	for (auto& PrimitivePair : PipelinedUpdateStatesRendering.Primitives)
 	{
 		FPrimitiveSceneInfo* RetrievedSceneInfo = Scene->GetPrimitiveSceneInfo(PrimitivePair.Value);
 		FPrimitiveSceneInfo* CachedSceneInfo = PrimitivePair.Key;
@@ -84,7 +74,7 @@ void FLateUpdateManager::Apply_RenderThread(FSceneInterface* Scene, const FTrans
 		RetrievedSceneInfo = Scene->GetPrimitiveSceneInfo(Index++);
 		while(RetrievedSceneInfo)
 		{
-			if (RetrievedSceneInfo->Proxy && UpdateStates[LateUpdateRenderReadIndex].Primitives.Contains(RetrievedSceneInfo) && UpdateStates[LateUpdateRenderReadIndex].Primitives[RetrievedSceneInfo] >= 0)
+			if (RetrievedSceneInfo->Proxy && PipelinedUpdateStatesRendering.Primitives.Contains(RetrievedSceneInfo) && PipelinedUpdateStatesRendering.Primitives[RetrievedSceneInfo] >= 0)
 			{
 				RetrievedSceneInfo->Proxy->ApplyLateUpdateTransform(RHICmdList, LateUpdateTransform);
 			}
@@ -103,7 +93,7 @@ void FLateUpdateManager::CacheSceneInfo(USceneComponent* Component)
 		FPrimitiveSceneInfo* PrimitiveSceneInfo = PrimitiveComponent->SceneProxy->GetPrimitiveSceneInfo();
 		if (PrimitiveSceneInfo && PrimitiveSceneInfo->IsIndexValid())
 		{
-			UpdateStates[LateUpdateGameWriteIndex].Primitives.Emplace(PrimitiveSceneInfo, PrimitiveSceneInfo->GetIndex());
+			PipelinedUpdateStatesGame.Primitives.Emplace(PrimitiveSceneInfo, PrimitiveSceneInfo->GetIndex());
 		}
 	}
 }
diff --git a/Engine/Source/Runtime/HeadMountedDisplay/Public/LateUpdateManager.h b/Engine/Source/Runtime/HeadMountedDisplay/Public/LateUpdateManager.h
--- a/Engine/Source/Runtime/HeadMountedDisplay/Public/LateUpdateManager.h
+++ b/Engine/Source/Runtime/HeadMountedDisplay/Public/LateUpdateManager.h
@@ -13,14 +13,14 @@ class FPrimitiveSceneInfo;
 class FLateUpdateManager
 {
 public:
-	HEADMOUNTEDDISPLAY_API FLateUpdateManager();
+	HEADMOUNTEDDISPLAY_API FLateUpdateManager() {}
 	virtual ~FLateUpdateManager() {}
 
 	/** Setup state for applying the render thread late update */
 	HEADMOUNTEDDISPLAY_API void Setup(const FTransform& ParentToWorld, USceneComponent* Component, bool bSkipLateUpdate);
 
 	/** Returns true if the LateUpdateSetup data is stale. */
-	bool GetSkipLateUpdate_RenderThread() const { return UpdateStates[LateUpdateRenderReadIndex].bSkip; }
+	bool GetSkipLateUpdate_RenderThread() const { return PipelinedUpdateStatesRendering.bSkip; }
 
 	/** Apply the late update delta to the cached components */
 	HEADMOUNTEDDISPLAY_API void Apply_RenderThread(FSceneInterface* Scene, const FTransform& OldRelativeTransform, const FTransform& NewRelativeTransform);
@@ -37,7 +37,6 @@ private:
 		FLateUpdateState()
 			: ParentToWorld(FTransform::Identity)
 			, bSkip(false)
-			, TrackingNumber(-1)
 		{}
 
 		/** Parent world transform used to reconstruct new world transforms for late update scene proxies */
@@ -46,12 +45,9 @@ private:
 		TMap<FPrimitiveSceneInfo*, int32> Primitives;
 		/** Late Update Info Stale, if this is found true do not late update */
 		bool bSkip;
-		/** Frame tracking number - used to flag if the game and render threads get badly out of sync */
-		int64 TrackingNumber;
 	};
 
-	FLateUpdateState UpdateStates[2];
-	int32 LateUpdateGameWriteIndex;
-	int32 LateUpdateRenderReadIndex;
+	FLateUpdateState PipelinedUpdateStatesGame;
+	FLateUpdateState PipelinedUpdateStatesRendering;
 };
 
